Fixed CIS

diff --git a/Engine/Source/Developer/NewLogVisualizer/Private/LogVisualizer.h b/Engine/Source/Developer/NewLogVisualizer/Private/LogVisualizer.h
--- a/Engine/Source/Developer/NewLogVisualizer/Private/LogVisualizer.h
+++ b/Engine/Source/Developer/NewLogVisualizer/Private/LogVisualizer.h
@@ -3,12 +3,12 @@
 
 /* Dependencies
 *****************************************************************************/
-
+#include "Engine.h"
 #include "Messaging.h"
 #include "ModuleManager.h"
 #include "SlateBasics.h"
 #include "SlateStyle.h"
-
+#include "EditorStyle.h"
 
 /* Private includes
 *****************************************************************************/
@@ -33,6 +33,7 @@ struct FVisualLoggerEvents
 
 struct IVisualLoggerInterface
 {
+	virtual ~IVisualLoggerInterface() {}
 	virtual bool HasValidCategories(TArray<FVisualLoggerCategoryVerbosityPair> Categories) = 0;
 	virtual bool IsValidCategory(const FString& InCategoryName, TEnumAsByte<ELogVerbosity::Type> Verbosity = ELogVerbosity::All) = 0;
 	virtual bool IsValidCategory(const FString& InGraphName, const FString& InDataName, TEnumAsByte<ELogVerbosity::Type> Verbosity = ELogVerbosity::All) = 0;
diff --git a/Engine/Source/Developer/NewLogVisualizer/Private/SVisualLogger.cpp b/Engine/Source/Developer/NewLogVisualizer/Private/SVisualLogger.cpp
--- a/Engine/Source/Developer/NewLogVisualizer/Private/SVisualLogger.cpp
+++ b/Engine/Source/Developer/NewLogVisualizer/Private/SVisualLogger.cpp
@@ -45,6 +45,8 @@ struct FVisualLoggerInterface : public IVisualLoggerInterface
 		VisualLoggerEvents = InVisualLoggerEvents;
 	}
 
+	virtual ~FVisualLoggerInterface() {}
+
 	virtual bool HasValidCategories(TArray<FVisualLoggerCategoryVerbosityPair> Categories) override
 	{
 		for (const auto& CurrentCategory : Categories)
@@ -119,7 +121,10 @@ protected:
 SVisualLogger::FVisualLoggerDevice::FVisualLoggerDevice(SVisualLogger* InOwner)
 	:Owner(InOwner)
 {
+}
 
+SVisualLogger::FVisualLoggerDevice::~FVisualLoggerDevice() 
+{
 }
 
 void SVisualLogger::FVisualLoggerDevice::Serialize(const class UObject* LogOwner, FName OwnerName, const FVisualLogEntry& LogEntry)
@@ -616,3 +621,4 @@ void SVisualLogger::OnFiltersSearchChanged(const FText& Filter)
 	LogsList->OnFiltersChanged();
 	MainView->OnFiltersChanged();
 }
+#undef LOCTEXT_NAMESPACE
diff --git a/Engine/Source/Developer/NewLogVisualizer/Private/SVisualLogger.h b/Engine/Source/Developer/NewLogVisualizer/Private/SVisualLogger.h
--- a/Engine/Source/Developer/NewLogVisualizer/Private/SVisualLogger.h
+++ b/Engine/Source/Developer/NewLogVisualizer/Private/SVisualLogger.h
@@ -39,6 +39,7 @@ public:
 	{
 	public:
 		FVisualLoggerDevice(SVisualLogger* Owner);
+		virtual ~FVisualLoggerDevice();
 		virtual void Serialize(const class UObject* LogOwner, FName OwnerName, const FVisualLogEntry& LogEntry) override;
 
 	protected:
diff --git a/Engine/Source/Developer/NewLogVisualizer/Private/SVisualLoggerLogsList.cpp b/Engine/Source/Developer/NewLogVisualizer/Private/SVisualLoggerLogsList.cpp
--- a/Engine/Source/Developer/NewLogVisualizer/Private/SVisualLoggerLogsList.cpp
+++ b/Engine/Source/Developer/NewLogVisualizer/Private/SVisualLoggerLogsList.cpp
@@ -166,3 +166,4 @@ void SVisualLoggerLogsList::OnItemSelectionChanged(const FVisualLogDevice::FVisu
 	//SetCurrentViewedTime(LogEntry.Entry.TimeStamp);
 	LogsLinesWidget->RequestListRefresh();
 }
+#undef LOCTEXT_NAMESPACE
diff --git a/Engine/Source/Developer/NewLogVisualizer/Private/SVisualLoggerView.cpp b/Engine/Source/Developer/NewLogVisualizer/Private/SVisualLoggerView.cpp
--- a/Engine/Source/Developer/NewLogVisualizer/Private/SVisualLoggerView.cpp
+++ b/Engine/Source/Developer/NewLogVisualizer/Private/SVisualLoggerView.cpp
@@ -9,35 +9,35 @@
 
 #define LOCTEXT_NAMESPACE "SVisualLoggerFilters"
 
-template <typename ItemType>
-class SLogListView : public SListView<ItemType>
-{
-public:
-	virtual FReply OnMouseWheel(const FGeometry& MyGeometry, const FPointerEvent& MouseEvent) override
-	{
-		if (!MouseEvent.IsLeftShiftDown())
-		{
-			return SListView<ItemType>::OnMouseWheel(MyGeometry, MouseEvent);
-		};
-
-		return FReply::Unhandled();
-	}
-
-	void RefreshList()
-	{
-		const TArray<ItemType>& ItemsSourceRef = (*this->ItemsSource);
-
-		for (int32 Index = 0; Index < ItemsSourceRef.Num(); ++Index)
-		{
-			TSharedPtr< SLogsTableRow > TableRow = StaticCastSharedPtr< SLogsTableRow >(this->WidgetGenerator.GetWidgetForItem(ItemsSourceRef[Index]));
-			if (TableRow.IsValid())
-			{
-				TableRow->UpdateEntries();
-			}
-		}
-
-	}
-};
+//template <typename ItemType>
+//class SLogListView : public SListView<ItemType>
+//{
+//public:
+//	virtual FReply OnMouseWheel(const FGeometry& MyGeometry, const FPointerEvent& MouseEvent) override
+//	{
+//		if (!MouseEvent.IsLeftShiftDown())
+//		{
+//			return SListView<ItemType>::OnMouseWheel(MyGeometry, MouseEvent);
+//		};
+//
+//		return FReply::Unhandled();
+//	}
+//
+//	void RefreshList()
+//	{
+//		const TArray<ItemType>& ItemsSourceRef = (*this->ItemsSource);
+//
+//		for (int32 Index = 0; Index < ItemsSourceRef.Num(); ++Index)
+//		{
+//			TSharedPtr< SLogsTableRow > TableRow = StaticCastSharedPtr< SLogsTableRow >(this->WidgetGenerator.GetWidgetForItem(ItemsSourceRef[Index]));
+//			if (TableRow.IsValid())
+//			{
+//				TableRow->UpdateEntries();
+//			}
+//		}
+//
+//	}
+//};
 
 void SVisualLoggerView::GetTimelines(TArray<TSharedPtr<STimeline> >& OutList, bool bOnlySelectedOnes)
 {
@@ -230,3 +230,4 @@ void SVisualLoggerView::OnFiltersChanged()
 {
 	TimelinesContainer->OnFiltersChanged();
 }
+#undef LOCTEXT_NAMESPACE
diff --git a/Engine/Source/Runtime/Engine/Public/DebugRenderSceneProxy.h b/Engine/Source/Runtime/Engine/Public/DebugRenderSceneProxy.h
--- a/Engine/Source/Runtime/Engine/Public/DebugRenderSceneProxy.h
+++ b/Engine/Source/Runtime/Engine/Public/DebugRenderSceneProxy.h
@@ -212,11 +212,11 @@ public:
 
 		}
 
+		float Radius;
 		FVector Location; //Center pointer of the base of the cylinder.
-		FVector X, Y, Z; //X, Y, and Z alignment axes to draw along.
 		FColor Color;
-		float Radius;
 		float HalfHeight;
+		FVector X, Y, Z; //X, Y, and Z alignment axes to draw along.
 	};
 
 	TArray<FDebugLine> Lines;
