Inverted Shift and Conytroll keys when using with Wheel in LogVisualizer tool. It's consistent with old tool now.

diff --git a/Engine/Source/Developer/NewLogVisualizer/Private/STimelinesContainer.cpp b/Engine/Source/Developer/NewLogVisualizer/Private/STimelinesContainer.cpp
--- a/Engine/Source/Developer/NewLogVisualizer/Private/STimelinesContainer.cpp
+++ b/Engine/Source/Developer/NewLogVisualizer/Private/STimelinesContainer.cpp
@@ -116,7 +116,7 @@ FReply STimelinesContainer::OnMouseMove(const FGeometry& MyGeometry, const FPoin
 
 FReply STimelinesContainer::OnMouseWheel(const FGeometry& MyGeometry, const FPointerEvent& MouseEvent)
 {
-	if (MouseEvent.IsLeftShiftDown())
+	if (MouseEvent.IsLeftControlDown() || MouseEvent.IsLeftShiftDown())
 	{
 		return TimeSliderController->OnMouseWheel(SharedThis(this), MyGeometry, MouseEvent);
 	}
diff --git a/Engine/Source/Developer/NewLogVisualizer/Private/TimeSliderController.cpp b/Engine/Source/Developer/NewLogVisualizer/Private/TimeSliderController.cpp
--- a/Engine/Source/Developer/NewLogVisualizer/Private/TimeSliderController.cpp
+++ b/Engine/Source/Developer/NewLogVisualizer/Private/TimeSliderController.cpp
@@ -463,9 +463,9 @@ FReply FSequencerTimeSliderController::OnMouseWheel( TSharedRef<SWidget> WidgetO
 {
 	FReply ReturnValue = FReply::Unhandled();;
 
-	if (MouseEvent.IsLeftControlDown())
+	if (MouseEvent.IsLeftShiftDown())
 	{
-		const float ZoomDelta = -0.1f * MouseEvent.GetWheelDelta();
+		const float ZoomDelta = 0.1f * MouseEvent.GetWheelDelta();
 		TimeSliderArgs.CursorSize.Set(FMath::Clamp(TimeSliderArgs.CursorSize.Get() + ZoomDelta, 0.0f, 1.0f));
 
 		ReturnValue = FReply::Handled();
