Added debug camera to LogVisualizer, to fly aroung during game pause #ue4
Fixed movement on timeline using keyboard in LogVisualizer #ue4
Added red color for items on timieline with Erro in Visual Logs #ue4.

diff --git a/Engine/Source/Developer/NewLogVisualizer/Private/STimelineBar.cpp b/Engine/Source/Developer/NewLogVisualizer/Private/STimelineBar.cpp
--- a/Engine/Source/Developer/NewLogVisualizer/Private/STimelineBar.cpp
+++ b/Engine/Source/Developer/NewLogVisualizer/Private/STimelineBar.cpp
@@ -35,31 +35,86 @@ FReply STimelineBar::OnMouseMove(const FGeometry& MyGeometry, const FPointerEven
 
 FReply STimelineBar::OnKeyDown(const FGeometry& MyGeometry, const FKeyEvent& InKeyEvent)
 {
+	FReply ReturnValue = FReply::Unhandled();
 	if (TimelineOwner.Pin()->IsSelected() == false)
 	{
-		return FReply::Unhandled();
+		return ReturnValue;
 	}
 
 	if (CurrentItemIndex != INDEX_NONE)
 	{
 		TRange<float> LocalViewRange = TimeSliderController->GetTimeSliderArgs().ViewRange.Get();
+		float RangeSize = LocalViewRange.Size<float>();
+
 		const FKey Key = InKeyEvent.GetKey();
-		if (Key == EKeys::Left || Key == EKeys::Right)
+		int32 NewItemIndex = CurrentItemIndex;
+		auto &Entries = TimelineOwner.Pin()->GetEntries();
+		if (Key == EKeys::Home)
+		{
+			for (int32 Index = 0; Index < Entries.Num(); Index++)
+			{
+				if (TimelineOwner.Pin()->IsEntryHidden(Entries[Index]) == false)
+				{
+					NewItemIndex = Index;
+					break;
+				}
+			}
+			ReturnValue = FReply::Handled();
+		}
+		else if (Key == EKeys::End)
 		{
-			int32 MoveBy = Key == EKeys::Left ? -1 : 1;
-			if (InKeyEvent.IsLeftControlDown())
+			for (int32 Index = Entries.Num()-1; Index >= 0; Index--)
 			{
-				MoveBy *= 10;
+				if (TimelineOwner.Pin()->IsEntryHidden(Entries[Index]) == false)
+				{
+					NewItemIndex = Index;
+					break;
+				}
 			}
+			ReturnValue = FReply::Handled();
+		}
+		else if (Key == EKeys::Left || Key == EKeys::Right)
+		{
+			int32 Direction = Key == EKeys::Left ? -1 : 1;
+			int32 MoveBy = InKeyEvent.IsLeftControlDown() ? InKeyEvent.IsLeftShiftDown() ? 20 : 10 : 1;
 
-			auto &Entries = TimelineOwner.Pin()->GetEntries();
-			uint32 NewItemIndex = FMath::Clamp(CurrentItemIndex + MoveBy, 0, Entries.Num() - 1);
-			SnapScrubPosition(Entries[NewItemIndex].Entry.TimeStamp);
-			return FReply::Handled();
+			int32 Index = 0;
+			while (true)
+			{
+				NewItemIndex += Direction;
+				if (Entries.IsValidIndex(NewItemIndex))
+				{
+					auto& CurrentEntryItem = Entries[NewItemIndex];
+					if (TimelineOwner.Pin()->IsEntryHidden(CurrentEntryItem) == false && ++Index == MoveBy)
+					{
+						break;
+					}
+				}
+				else
+				{
+					NewItemIndex = FMath::Clamp(NewItemIndex, 0, Entries.Num() - 1);
+					break;
+				}
+			}
+			ReturnValue = FReply::Handled();
+		}
+
+		if (NewItemIndex != CurrentItemIndex)
+		{
+			float NewTimeStamp = Entries[NewItemIndex].Entry.TimeStamp;
+			SnapScrubPosition(NewTimeStamp);
+			if (NewTimeStamp < LocalViewRange.GetLowerBoundValue())
+			{
+				TimeSliderController->GetTimeSliderArgs().ViewRange.Set(TRange<float>(NewTimeStamp, NewTimeStamp + RangeSize));
+			}
+			else if (NewTimeStamp > LocalViewRange.GetUpperBoundValue())
+			{
+				TimeSliderController->GetTimeSliderArgs().ViewRange.Set(TRange<float>(NewTimeStamp - RangeSize, NewTimeStamp));
+			}
 		}
 	}
 
-	return FReply::Unhandled();
+	return ReturnValue;
 }
 
 uint32 STimelineBar::GetClosestItem(float Time) const
@@ -162,13 +217,18 @@ int32 STimelineBar::OnPaint(const FPaintArgs& Args, const FGeometry& AllottedGeo
 		TimelineOwner.Pin()->IsSelected() ? FLinearColor(.2f, .2f, .2f, 0.5f) : FLinearColor(.1f, .1f, .1f, 0.5f)
 		);
 
+	const FSlateBrush* FillImage = FLogVisualizerStyle::Get().GetBrush("LogVisualizer.LogBar.EntryDefault");
+	static const FColor CurrentTimeColor(140, 255, 255, 255);
+	static const FColor ErrorTimeColor(255, 0, 0, 255);
+	static const FColor SelectedBarColor(255, 255, 255, 255);
+	const FSlateBrush* SelectedFillImage = FLogVisualizerStyle::Get().GetBrush("LogVisualizer.LogBar.Selected");
+
+	const ESlateDrawEffect::Type DrawEffects = ESlateDrawEffect::None;// bEnabled ? ESlateDrawEffect::None : ESlateDrawEffect::DisabledEffect;
+
+	TArray<float> ErrorTimes;
 	auto &Entries = TimelineOwner.Pin()->GetEntries();
 	for (int32 Index = 0; Index < Entries.Num(); Index++)
 	{
-		const FSlateBrush* FillImage = FLogVisualizerStyle::Get().GetBrush("LogVisualizer.LogBar.EntryDefault");
-
-		static const FColor CurrentTimeColor(140, 255, 255, 255);
-		const ESlateDrawEffect::Type DrawEffects = ESlateDrawEffect::None;// bEnabled ? ESlateDrawEffect::None : ESlateDrawEffect::DisabledEffect;
 		int32 CurrentIndex = Index;
 		float CurrentTime = Entries[CurrentIndex].Entry.TimeStamp;
 		if (CurrentTime < LocalViewRange.GetLowerBoundValue() || CurrentTime > LocalViewRange.GetUpperBoundValue())
@@ -181,6 +241,16 @@ int32 STimelineBar::OnPaint(const FPaintArgs& Args, const FGeometry& AllottedGeo
 			continue;
 		}
 
+		const TArray<FVisualLogLine>& LogLines = Entries[CurrentIndex].Entry.LogLines;
+		for (const FVisualLogLine& CurrentLine : LogLines)
+		{
+			if (CurrentLine.Verbosity <= ELogVerbosity::Error)
+			{
+				ErrorTimes.AddUnique(CurrentTime);
+				break;
+			}
+		}
+
 		float LinePos = (CurrentTime - LocalViewRange.GetLowerBoundValue()) * PixelsPerInput;
 
 		FSlateDrawElement::MakeBox(
@@ -196,6 +266,24 @@ int32 STimelineBar::OnPaint(const FPaintArgs& Args, const FGeometry& AllottedGeo
 			);
 	}
 
+	float ErrorBoxWidth = BoxWidth+2;
+	for (auto CurrentTime : ErrorTimes)
+	{
+		float LinePos = (CurrentTime - LocalViewRange.GetLowerBoundValue()) * PixelsPerInput;
+
+		FSlateDrawElement::MakeBox(
+			OutDrawElements,
+			RetLayerId++,
+			AllottedGeometry.ToPaintGeometry(
+			FVector2D(LinePos - ErrorBoxWidth * 0.5f, 0.0f),
+			FVector2D(ErrorBoxWidth, AllottedGeometry.Size.Y)),
+			FillImage,
+			MyClippingRect,
+			DrawEffects,
+			ErrorTimeColor
+			);
+	}
+
 	uint32 BestItemIndex = GetClosestItem(LocalScrubPosition);
 
 	if (BestItemIndex != INDEX_NONE && TimelineOwner.Pin()->IsSelected())
@@ -206,10 +294,6 @@ int32 STimelineBar::OnPaint(const FPaintArgs& Args, const FGeometry& AllottedGeo
 			VisualLoggerInterface->GetVisualLoggerEvents().OnItemSelectionChanged.ExecuteIfBound(Entries[BestItemIndex]);
 		}
 
-		static const FColor SelectedBarColor(255, 255, 255, 255);
-		static const FColor CurrentTimeColor(140, 255, 255, 255);
-
-		const FSlateBrush* FillImage = FLogVisualizerStyle::Get().GetBrush("LogVisualizer.LogBar.Selected");
 		float CurrentTime = Entries[BestItemIndex].Entry.TimeStamp;
 		float LinePos = (CurrentTime - LocalViewRange.GetLowerBoundValue()) * PixelsPerInput;
 
@@ -220,7 +304,7 @@ int32 STimelineBar::OnPaint(const FPaintArgs& Args, const FGeometry& AllottedGeo
 			AllottedGeometry.ToPaintGeometry(
 			FVector2D(LinePos - BoxWidth * 0.5f, 0.0f),
 			FVector2D(BoxWidth, AllottedGeometry.Size.Y)),
-			FillImage,
+			SelectedFillImage,
 			MyClippingRect,
 			ESlateDrawEffect::None,
 			SelectedBarColor
diff --git a/Engine/Source/Developer/NewLogVisualizer/Private/SVisualLogger.cpp b/Engine/Source/Developer/NewLogVisualizer/Private/SVisualLogger.cpp
--- a/Engine/Source/Developer/NewLogVisualizer/Private/SVisualLogger.cpp
+++ b/Engine/Source/Developer/NewLogVisualizer/Private/SVisualLogger.cpp
@@ -7,7 +7,7 @@
 #include "VisualLoggerCanvasRenderer.h"
 #include "DesktopPlatformModule.h"
 #include "MainFrame.h"
-
+#include "VisualLoggerCameraController.h"
 #if WITH_EDITOR
 #	include "Editor/UnrealEd/Public/EditorComponents.h"
 #	include "Editor/UnrealEd/Public/EditorReimportHandler.h"
@@ -182,6 +182,7 @@ void SVisualLogger::Construct(const FArguments& InArgs, const TSharedRef<SDockTa
 	ActionList.MapAction(Commands.Resume, FExecuteAction::CreateRaw(this, &SVisualLogger::HandleResumeCommandExecute), FCanExecuteAction::CreateRaw(this, &SVisualLogger::HandleResumeCommandCanExecute), FIsActionChecked(), FIsActionButtonVisible::CreateRaw(this, &SVisualLogger::HandleResumeCommandIsVisible));
 	ActionList.MapAction(Commands.Load, FExecuteAction::CreateRaw(this, &SVisualLogger::HandleLoadCommandExecute), FCanExecuteAction::CreateRaw(this, &SVisualLogger::HandleLoadCommandCanExecute), FIsActionChecked(), FIsActionButtonVisible::CreateRaw(this, &SVisualLogger::HandleLoadCommandCanExecute));
 	ActionList.MapAction(Commands.Save, FExecuteAction::CreateRaw(this, &SVisualLogger::HandleSaveCommandExecute), FCanExecuteAction::CreateRaw(this, &SVisualLogger::HandleSaveCommandCanExecute), FIsActionChecked(), FIsActionButtonVisible::CreateRaw(this, &SVisualLogger::HandleSaveCommandCanExecute));
+	ActionList.MapAction(Commands.FreeCamera, FExecuteAction::CreateRaw(this, &SVisualLogger::HandleCameraCommandExecute), FCanExecuteAction::CreateRaw(this, &SVisualLogger::HandleCameraCommandCanExecute), FIsActionChecked(), FIsActionButtonVisible::CreateRaw(this, &SVisualLogger::HandleCameraCommandCanExecute));
 
 	// Tab Spawners
 	TabManager = FGlobalTabmanager::Get()->NewTabManager(ConstructUnderMajorTab);
@@ -391,17 +392,6 @@ void SVisualLogger::HandlePauseCommandExecute()
 	if (World != NULL)
 	{
 		World->bPlayersOnlyPending = true;
-		//// switch debug cam on
-		//CameraController = ALogVisualizerCameraController::EnableCamera(World);
-		//if (CameraController.IsValid())
-		//{
-		//	CameraController->OnActorSelected = ALogVisualizerCameraController::FActorSelectedDelegate::CreateSP(
-		//		this, &SLogVisualizer::CameraActorSelected
-		//		);
-		//	CameraController->OnIterateLogEntries = ALogVisualizerCameraController::FLogEntryIterationDelegate::CreateSP(
-		//		this, &SLogVisualizer::IncrementCurrentLogIndex
-		//		);
-		//}
 	}
 }
 
@@ -423,8 +413,6 @@ void SVisualLogger::HandleResumeCommandExecute()
 	{
 		World->bPlayersOnly = false;
 		World->bPlayersOnlyPending = false;
-
-		//ALogVisualizerCameraController::DisableCamera(World);
 	}
 }
 
@@ -433,6 +421,31 @@ bool SVisualLogger::HandleResumeCommandIsVisible() const
 	return HandleResumeCommandCanExecute();
 }
 
+bool SVisualLogger::HandleCameraCommandCanExecute() const
+{
+	UWorld* World = VisualLoggerInterface->GetWorld();
+	return FVisualLogger::Get().IsRecording() && World && !World->bPlayersOnly && !World->bPlayersOnlyPending;
+}
+
+void SVisualLogger::HandleCameraCommandExecute()
+{
+	UWorld* World = VisualLoggerInterface->GetWorld();
+	if (AVisualLoggerCameraController::IsEnabled(World))
+	{
+		AVisualLoggerCameraController::DisableCamera(World);
+	}
+	else
+	{
+		// switch debug cam on
+		CameraController = AVisualLoggerCameraController::EnableCamera(World);
+		if (CameraController.IsValid())
+		{
+			//CameraController->OnActorSelected = ALogVisualizerCameraController::FActorSelectedDelegate::CreateSP(this, &SVisualLogger::CameraActorSelected);
+			//CameraController->OnIterateLogEntries = ALogVisualizerCameraController::FLogEntryIterationDelegate::CreateSP(this, &SVisualLogger::IncrementCurrentLogIndex);
+		}
+	}
+}
+
 bool SVisualLogger::HandleLoadCommandCanExecute() const
 {
 	return true;
diff --git a/Engine/Source/Developer/NewLogVisualizer/Private/SVisualLogger.h b/Engine/Source/Developer/NewLogVisualizer/Private/SVisualLogger.h
--- a/Engine/Source/Developer/NewLogVisualizer/Private/SVisualLogger.h
+++ b/Engine/Source/Developer/NewLogVisualizer/Private/SVisualLogger.h
@@ -82,6 +82,9 @@ public:
 	bool HandleSaveCommandCanExecute() const;
 	void HandleSaveCommandExecute();
 
+	bool HandleCameraCommandCanExecute() const;
+	void HandleCameraCommandExecute();
+
 	TSharedPtr<SVisualLoggerFilters> GetVisualLoggerFilters() { return VisualLoggerFilters; }
 
 protected:
@@ -94,6 +97,7 @@ protected:
 	// Visual Logger device to get and collect logs.
 	TSharedPtr<FVisualLoggerDevice> InternalDevice;
 
+	TWeakObjectPtr<class AVisualLoggerCameraController> CameraController;
 	TSharedPtr<struct FVisualLoggerCanvasRenderer> VisualLoggerCanvasRenderer;
 	TSharedPtr<IVisualLoggerInterface> VisualLoggerInterface;
 	mutable TSharedPtr<SVisualLoggerFilters> VisualLoggerFilters;
diff --git a/Engine/Source/Developer/NewLogVisualizer/Private/SVisualLoggerToolbar.cpp b/Engine/Source/Developer/NewLogVisualizer/Private/SVisualLoggerToolbar.cpp
--- a/Engine/Source/Developer/NewLogVisualizer/Private/SVisualLoggerToolbar.cpp
+++ b/Engine/Source/Developer/NewLogVisualizer/Private/SVisualLoggerToolbar.cpp
@@ -26,6 +26,8 @@ TSharedRef<SWidget> SVisualLoggerToolbar::MakeToolbar(const TSharedRef<FUIComman
 		ToolBarBuilder.AddToolBarButton(FVisualLoggerCommands::Get().Pause, NAME_None, LOCTEXT("PauseGame", "Pause"), TAttribute<FText>(), FSlateIcon(FLogVisualizerStyle::Get().GetStyleSetName(), TEXT("Toolbar.Pause")));
 		ToolBarBuilder.AddToolBarButton(FVisualLoggerCommands::Get().Resume, NAME_None, LOCTEXT("ResumeGame", "Resume"), TAttribute<FText>(), FSlateIcon(FLogVisualizerStyle::Get().GetStyleSetName(), TEXT("Toolbar.Resume")));
 
+		ToolBarBuilder.AddSeparator();
+		ToolBarBuilder.AddToolBarButton(FVisualLoggerCommands::Get().FreeCamera, NAME_None, LOCTEXT("FreeCamera", "Camera"), TAttribute<FText>(), FSlateIcon(FLogVisualizerStyle::Get().GetStyleSetName(), TEXT("Toolbar.Camera")));
 		ToolBarBuilder.AddSeparator();
 		ToolBarBuilder.AddToolBarButton(FVisualLoggerCommands::Get().Load, NAME_None, LOCTEXT("Load", "Load"), TAttribute<FText>(), FSlateIcon(FLogVisualizerStyle::Get().GetStyleSetName(), TEXT("Toolbar.Load")));
 		ToolBarBuilder.AddToolBarButton(FVisualLoggerCommands::Get().Save, NAME_None, LOCTEXT("Save", "Save"), TAttribute<FText>(), FSlateIcon(FLogVisualizerStyle::Get().GetStyleSetName(), TEXT("Toolbar.Save")));
diff --git a/Engine/Source/Developer/NewLogVisualizer/Private/SVisualLoggerView.cpp b/Engine/Source/Developer/NewLogVisualizer/Private/SVisualLoggerView.cpp
--- a/Engine/Source/Developer/NewLogVisualizer/Private/SVisualLoggerView.cpp
+++ b/Engine/Source/Developer/NewLogVisualizer/Private/SVisualLoggerView.cpp
@@ -9,36 +9,6 @@
 
 #define LOCTEXT_NAMESPACE "SVisualLoggerFilters"
 
-//template <typename ItemType>
-//class SLogListView : public SListView<ItemType>
-//{
-//public:
-//	virtual FReply OnMouseWheel(const FGeometry& MyGeometry, const FPointerEvent& MouseEvent) override
-//	{
-//		if (!MouseEvent.IsLeftShiftDown())
-//		{
-//			return SListView<ItemType>::OnMouseWheel(MyGeometry, MouseEvent);
-//		};
-//
-//		return FReply::Unhandled();
-//	}
-//
-//	void RefreshList()
-//	{
-//		const TArray<ItemType>& ItemsSourceRef = (*this->ItemsSource);
-//
-//		for (int32 Index = 0; Index < ItemsSourceRef.Num(); ++Index)
-//		{
-//			TSharedPtr< SLogsTableRow > TableRow = StaticCastSharedPtr< SLogsTableRow >(this->WidgetGenerator.GetWidgetForItem(ItemsSourceRef[Index]));
-//			if (TableRow.IsValid())
-//			{
-//				TableRow->UpdateEntries();
-//			}
-//		}
-//
-//	}
-//};
-
 void SVisualLoggerView::GetTimelines(TArray<TSharedPtr<STimeline> >& OutList, bool bOnlySelectedOnes)
 {
 	OutList = bOnlySelectedOnes ? TimelinesContainer->GetSelectedNodes() : TimelinesContainer->GetAllNodes();
diff --git a/Engine/Source/Developer/NewLogVisualizer/Private/VisualLoggerCameraController.cpp b/Engine/Source/Developer/NewLogVisualizer/Private/VisualLoggerCameraController.cpp
new file mode 100644
--- /dev/null
+++ b/Engine/Source/Developer/NewLogVisualizer/Private/VisualLoggerCameraController.cpp
@@ -0,0 +1,153 @@
+// Copyright 1998-2014 Epic Games, Inc. All Rights Reserved.
+
+#include "LogVisualizer.h"
+#include "VisualLoggerCameraController.h"
+#include "VisualLoggerHUD.h"
+#include "Runtime/Engine/Classes/Kismet/GameplayStatics.h"
+#include "GameFramework/PlayerInput.h"
+
+
+//----------------------------------------------------------------------//
+// AVisualLoggerCameraController
+//----------------------------------------------------------------------//
+TWeakObjectPtr<AVisualLoggerCameraController> AVisualLoggerCameraController::Instance;
+
+AVisualLoggerCameraController::AVisualLoggerCameraController(const FObjectInitializer& ObjectInitializer)
+	: Super(ObjectInitializer)
+{
+	SpeedScale = 1.5f;
+}
+
+void AVisualLoggerCameraController::SetupInputComponent()
+{
+	Super::SetupInputComponent();
+
+	static bool bBindingsAdded = false;
+	if (!bBindingsAdded)
+	{
+		bBindingsAdded = true;
+
+		UPlayerInput::AddEngineDefinedActionMapping(FInputActionKeyMapping("LogCamera_NextEntry", EKeys::LeftMouseButton, true));
+		UPlayerInput::AddEngineDefinedActionMapping(FInputActionKeyMapping("LogCamera_PrevEntry", EKeys::RightMouseButton, true));
+	}
+
+	InputComponent->BindAction("LogCamera_NextEntry", IE_Pressed, this, &AVisualLoggerCameraController::ShowNextEntry);
+	InputComponent->BindAction("LogCamera_PrevEntry", IE_Pressed, this, &AVisualLoggerCameraController::ShowPrevEntry);
+}
+
+void AVisualLoggerCameraController::ShowNextEntry()
+{
+	OnIterateLogEntries.ExecuteIfBound(1);
+}
+
+void AVisualLoggerCameraController::ShowPrevEntry()
+{
+	OnIterateLogEntries.ExecuteIfBound(-1);
+}
+
+void AVisualLoggerCameraController::PostInitializeComponents()
+{
+	Super::PostInitializeComponents();
+
+	// if hud is existing, delete it and create new hud for debug camera
+	if ( MyHUD != NULL )
+	{
+		MyHUD->Destroy();
+	}
+	FActorSpawnParameters SpawnInfo;
+	SpawnInfo.Owner = this;
+	SpawnInfo.Instigator = Instigator;
+	SpawnInfo.bNoCollisionFail = true;
+	MyHUD = GetWorld()->SpawnActor<AVisualLoggerHUD>(SpawnInfo);
+
+	/** TO DO: somehow this two configuration settings are set to zero for this controller, have to fix it later */
+	InputYawScale=2.5;
+	InputPitchScale=-1.75;
+}
+
+
+AVisualLoggerCameraController* AVisualLoggerCameraController::EnableCamera(UWorld* InWorld)
+{
+	APlayerController* PC = UGameplayStatics::GetPlayerController(InWorld, 0);
+	if (PC == NULL || PC->Player == NULL || PC->IsLocalPlayerController() == false
+		|| PC->GetWorld() == NULL)
+	{
+		return NULL;
+	}
+	
+	if (Instance.IsValid() == false)
+	{
+		// spawn if necessary
+		// and ungly @HACK to be able to spawn camera in game world rathen then
+		// in editor world (if running PIE). Hate it, but it works, and 
+		// this is a debugging tool		
+		FActorSpawnParameters SpawnInfo;
+		SpawnInfo.bNoCollisionFail = true;
+		SpawnInfo.Owner = PC->GetWorldSettings();
+		SpawnInfo.Instigator = PC->Instigator;
+		Instance = InWorld->SpawnActor<AVisualLoggerCameraController>( SpawnInfo );
+		Instance->Player = PC->Player;
+	}
+	
+	if (Instance.IsValid())
+	{
+		if (Instance.Get() != PC)
+		{
+			// set up new controller
+			Instance->OnActivate(PC);
+
+			// then switch to it
+			PC->Player->SwitchController( Instance.Get() );
+
+			Instance->ChangeState(NAME_Default);
+			Instance->ChangeState(NAME_Spectating);
+		}
+
+		return Instance.Get();
+	}
+
+	return NULL;
+}
+
+void AVisualLoggerCameraController::DisableCamera(UWorld* InWorld)
+{
+	check(InWorld);
+	for (FConstPlayerControllerIterator Iterator = InWorld->GetPlayerControllerIterator(); Iterator; ++Iterator)
+{
+		AVisualLoggerCameraController* VLogCam = Cast<AVisualLoggerCameraController>(*Iterator);
+		if (VLogCam && VLogCam->OriginalPlayer)
+	{
+		VLogCam->OriginalPlayer->SwitchController(VLogCam->OriginalControllerRef);
+		VLogCam->OnDeactivate(VLogCam->OriginalControllerRef);
+			break;
+		}
+	}
+}
+
+bool AVisualLoggerCameraController::IsEnabled(UWorld* InWorld)
+{
+	if (InWorld != NULL)
+	{
+		for (FConstPlayerControllerIterator Iterator = InWorld->GetPlayerControllerIterator(); Iterator; ++Iterator)
+		{
+			AVisualLoggerCameraController* VLogCam = Cast<AVisualLoggerCameraController>(*Iterator);
+			if (VLogCam && VLogCam->OriginalControllerRef && VLogCam->OriginalPlayer)
+			{
+				return true;
+			}
+		}
+	}
+
+	return false;
+}
+
+void AVisualLoggerCameraController::Select(FHitResult const& Hit) 
+{
+	Super::Select(Hit);
+
+	if (SelectedActor != PickedActor)
+	{
+		PickedActor = SelectedActor;
+		OnActorSelected.ExecuteIfBound(SelectedActor);
+	}
+}
diff --git a/Engine/Source/Developer/NewLogVisualizer/Private/VisualLoggerCameraController.h b/Engine/Source/Developer/NewLogVisualizer/Private/VisualLoggerCameraController.h
new file mode 100644
--- /dev/null
+++ b/Engine/Source/Developer/NewLogVisualizer/Private/VisualLoggerCameraController.h
@@ -0,0 +1,38 @@
+// Copyright 1998-2014 Epic Games, Inc. All Rights Reserved.
+
+#pragma once
+
+#include "Engine/DebugCameraController.h"
+#include "VisualLoggerCameraController.generated.h"
+
+UCLASS(config=Input, hidedropdown)
+class AVisualLoggerCameraController : public ADebugCameraController
+{
+public:
+	DECLARE_DELEGATE_OneParam(FActorSelectedDelegate, class AActor*);
+	DECLARE_DELEGATE_OneParam(FLogEntryIterationDelegate, int32);
+
+	GENERATED_UCLASS_BODY()
+
+	UPROPERTY()
+	class AActor* PickedActor;
+
+	virtual void PostInitializeComponents() override;
+
+	static AVisualLoggerCameraController* EnableCamera(UWorld* InWorld);
+	static void DisableCamera(UWorld* InWorld);
+	static bool IsEnabled(UWorld* InWorld);
+
+	virtual void Select( FHitResult const& Hit ) override;
+
+	void ShowNextEntry();
+	void ShowPrevEntry();
+
+	static TWeakObjectPtr<AVisualLoggerCameraController> Instance;
+
+	FActorSelectedDelegate OnActorSelected;
+	FLogEntryIterationDelegate OnIterateLogEntries;
+
+protected:
+	virtual void SetupInputComponent() override;
+};
diff --git a/Engine/Source/Developer/NewLogVisualizer/Private/VisualLoggerCommands.h b/Engine/Source/Developer/NewLogVisualizer/Private/VisualLoggerCommands.h
--- a/Engine/Source/Developer/NewLogVisualizer/Private/VisualLoggerCommands.h
+++ b/Engine/Source/Developer/NewLogVisualizer/Private/VisualLoggerCommands.h
@@ -24,6 +24,7 @@ public:
 		UI_COMMAND(StopRecording, "Stop", "Step over the current message", EUserInterfaceActionType::Button, FInputGesture());
 		UI_COMMAND(Pause, "Pause", "Stop the debugger", EUserInterfaceActionType::Button, FInputGesture());
 		UI_COMMAND(Resume, "Resume", "Stop the debugger", EUserInterfaceActionType::Button, FInputGesture());
+		UI_COMMAND(FreeCamera, "FreeCamera", "Enable free camera", EUserInterfaceActionType::ToggleButton, FInputGesture());
 		UI_COMMAND(Load, "Load", "Load external vlogs", EUserInterfaceActionType::Button, FInputGesture());
 		UI_COMMAND(Save, "Save", "Save selected data to vlog file", EUserInterfaceActionType::Button, FInputGesture());
 	}
@@ -34,6 +35,7 @@ public:
 	TSharedPtr<FUICommandInfo> StopRecording;
 	TSharedPtr<FUICommandInfo> Pause;
 	TSharedPtr<FUICommandInfo> Resume;
+	TSharedPtr<FUICommandInfo> FreeCamera;
 	TSharedPtr<FUICommandInfo> Load;
 	TSharedPtr<FUICommandInfo> Save;
 };
diff --git a/Engine/Source/Developer/NewLogVisualizer/Private/VisualLoggerHUD.cpp b/Engine/Source/Developer/NewLogVisualizer/Private/VisualLoggerHUD.cpp
new file mode 100644
--- /dev/null
+++ b/Engine/Source/Developer/NewLogVisualizer/Private/VisualLoggerHUD.cpp
@@ -0,0 +1,70 @@
+// Copyright 1998-2014 Epic Games, Inc. All Rights Reserved.
+
+#include "LogVisualizer.h"
+#include "VisualLoggerHUD.h"
+#include "VisualLoggerCameraController.h"
+#include "VisualLogger/VisualLogger.h"
+
+#define LOCTEXT_NAMESPACE "AVisualLoggerHUD"
+
+//----------------------------------------------------------------------//
+// AVisualLoggerHUD
+//----------------------------------------------------------------------//
+AVisualLoggerHUD::AVisualLoggerHUD(const FObjectInitializer& ObjectInitializer)
+	: Super(ObjectInitializer)
+{
+	bHidden = false;
+
+	TextRenderInfo.bEnableShadow = true;
+}
+
+bool AVisualLoggerHUD::DisplayMaterials( float X, float& Y, float DY, UMeshComponent* MeshComp )
+{
+	return false;
+}
+
+void AVisualLoggerHUD::PostRender()
+{
+	static const FColor TextColor(200, 200, 128, 255);
+
+	Super::Super::PostRender();
+
+	if (bShowHUD)
+	{
+		AVisualLoggerCameraController* DebugCamController = Cast<AVisualLoggerCameraController>(PlayerOwner);
+		if( DebugCamController != NULL )
+		{
+			FCanvasTextItem TextItem( FVector2D::ZeroVector, FText::GetEmpty(), GEngine->GetSmallFont(), TextColor);
+			TextItem.FontRenderInfo = TextRenderInfo;
+			float X = Canvas->SizeX * 0.025f+1;
+			float Y = Canvas->SizeX * 0.025f+1;
+
+			FVector const CamLoc = DebugCamController->PlayerCameraManager->GetCameraLocation();
+			FRotator const CamRot = DebugCamController->PlayerCameraManager->GetCameraRotation();
+
+			FCollisionQueryParams TraceParams(NAME_None, true, this);
+			FHitResult Hit;
+			bool bHit = GetWorld()->LineTraceSingle(Hit, CamLoc, CamRot.Vector() * 100000.f + CamLoc, ECC_Pawn, TraceParams);
+			if( bHit )
+			{
+				TextItem.Text = FText::FromString(FString::Printf(TEXT("Under cursor: '%s'"), *Hit.GetActor()->GetName()));
+				Canvas->DrawItem( TextItem, X, Y );
+				
+				DrawDebugLine( GetWorld(), Hit.Location, Hit.Location+Hit.Normal*30.f, FColor::White );
+			}
+			else
+			{
+				TextItem.Text = LOCTEXT("NotActorUnderCursor", "Not actor under cursor" );
+			}
+			Canvas->DrawItem( TextItem, X, Y );
+			Y += TextItem.DrawnSize.Y;
+
+			if (DebugCamController->PickedActor != NULL)
+			{
+				TextItem.Text = FText::FromString(FString::Printf(TEXT("Selected: '%s'"), *DebugCamController->PickedActor->GetName()));
+				Canvas->DrawItem( TextItem, X, Y );				
+			}
+		}
+	}
+}
+#undef LOCTEXT_NAMESPACE
diff --git a/Engine/Source/Developer/NewLogVisualizer/Private/VisualLoggerHUD.h b/Engine/Source/Developer/NewLogVisualizer/Private/VisualLoggerHUD.h
new file mode 100644
--- /dev/null
+++ b/Engine/Source/Developer/NewLogVisualizer/Private/VisualLoggerHUD.h
@@ -0,0 +1,22 @@
+// Copyright 1998-2014 Epic Games, Inc. All Rights Reserved.
+
+#pragma once
+#include "Engine/DebugCameraHUD.h"
+#include "VisualLoggerHUD.generated.h"
+
+UCLASS()
+class AVisualLoggerHUD : public ADebugCameraHUD
+{
+	GENERATED_UCLASS_BODY()
+
+	FFontRenderInfo TextRenderInfo;
+
+	virtual bool DisplayMaterials( float X, float& Y, float DY, UMeshComponent* MeshComp ) override;
+
+	// Begin AActor Interface
+	virtual void PostRender() override;
+	// End AActor Interface
+};
+
+
+
