Fixed camera toolbar button in LogVisualizer #ue4

diff --git a/Engine/Source/Developer/NewLogVisualizer/Private/SVisualLogger.cpp b/Engine/Source/Developer/NewLogVisualizer/Private/SVisualLogger.cpp
--- a/Engine/Source/Developer/NewLogVisualizer/Private/SVisualLogger.cpp
+++ b/Engine/Source/Developer/NewLogVisualizer/Private/SVisualLogger.cpp
@@ -182,7 +182,11 @@ void SVisualLogger::Construct(const FArguments& InArgs, const TSharedRef<SDockTa
 	ActionList.MapAction(Commands.Resume, FExecuteAction::CreateRaw(this, &SVisualLogger::HandleResumeCommandExecute), FCanExecuteAction::CreateRaw(this, &SVisualLogger::HandleResumeCommandCanExecute), FIsActionChecked(), FIsActionButtonVisible::CreateRaw(this, &SVisualLogger::HandleResumeCommandIsVisible));
 	ActionList.MapAction(Commands.Load, FExecuteAction::CreateRaw(this, &SVisualLogger::HandleLoadCommandExecute), FCanExecuteAction::CreateRaw(this, &SVisualLogger::HandleLoadCommandCanExecute), FIsActionChecked(), FIsActionButtonVisible::CreateRaw(this, &SVisualLogger::HandleLoadCommandCanExecute));
 	ActionList.MapAction(Commands.Save, FExecuteAction::CreateRaw(this, &SVisualLogger::HandleSaveCommandExecute), FCanExecuteAction::CreateRaw(this, &SVisualLogger::HandleSaveCommandCanExecute), FIsActionChecked(), FIsActionButtonVisible::CreateRaw(this, &SVisualLogger::HandleSaveCommandCanExecute));
-	ActionList.MapAction(Commands.FreeCamera, FExecuteAction::CreateRaw(this, &SVisualLogger::HandleCameraCommandExecute), FCanExecuteAction::CreateRaw(this, &SVisualLogger::HandleCameraCommandCanExecute), FIsActionChecked(), FIsActionButtonVisible::CreateRaw(this, &SVisualLogger::HandleCameraCommandCanExecute));
+	ActionList.MapAction(Commands.FreeCamera, 
+		FExecuteAction::CreateRaw(this, &SVisualLogger::HandleCameraCommandExecute), 
+		FCanExecuteAction::CreateRaw(this, &SVisualLogger::HandleCameraCommandCanExecute), 
+		FIsActionChecked::CreateRaw(this, &SVisualLogger::HandleCameraCommandIsChecked),
+		FIsActionButtonVisible::CreateRaw(this, &SVisualLogger::HandleCameraCommandCanExecute));
 
 	// Tab Spawners
 	TabManager = FGlobalTabmanager::Get()->NewTabManager(ConstructUnderMajorTab);
@@ -372,6 +376,12 @@ bool SVisualLogger::HandleStopRecordingCommandCanExecute() const
 void SVisualLogger::HandleStopRecordingCommandExecute()
 {
 	FVisualLogger::Get().SetIsRecording(false);
+
+	UWorld* World = VisualLoggerInterface->GetWorld();
+	if (AVisualLoggerCameraController::IsEnabled(World))
+	{
+		AVisualLoggerCameraController::DisableCamera(World);
+	}
 }
 
 
@@ -421,10 +431,16 @@ bool SVisualLogger::HandleResumeCommandIsVisible() const
 	return HandleResumeCommandCanExecute();
 }
 
+bool SVisualLogger::HandleCameraCommandIsChecked() const
+{
+	UWorld* World = VisualLoggerInterface->GetWorld();
+	return World && AVisualLoggerCameraController::IsEnabled(World);
+}
+
 bool SVisualLogger::HandleCameraCommandCanExecute() const
 {
 	UWorld* World = VisualLoggerInterface->GetWorld();
-	return FVisualLogger::Get().IsRecording() && World && !World->bPlayersOnly && !World->bPlayersOnlyPending;
+	return FVisualLogger::Get().IsRecording() && World && !World->bPlayersOnly && !World->bPlayersOnlyPending && World->IsPlayInEditor();
 }
 
 void SVisualLogger::HandleCameraCommandExecute()
diff --git a/Engine/Source/Developer/NewLogVisualizer/Private/SVisualLogger.h b/Engine/Source/Developer/NewLogVisualizer/Private/SVisualLogger.h
--- a/Engine/Source/Developer/NewLogVisualizer/Private/SVisualLogger.h
+++ b/Engine/Source/Developer/NewLogVisualizer/Private/SVisualLogger.h
@@ -84,6 +84,7 @@ public:
 
 	bool HandleCameraCommandCanExecute() const;
 	void HandleCameraCommandExecute();
+	bool HandleCameraCommandIsChecked() const;
 
 	TSharedPtr<SVisualLoggerFilters> GetVisualLoggerFilters() { return VisualLoggerFilters; }
 
