Tweaks for keyboard movement on Timelines in LogVisualizer #ue4

diff --git a/Engine/Source/Developer/NewLogVisualizer/Private/STimelineBar.cpp b/Engine/Source/Developer/NewLogVisualizer/Private/STimelineBar.cpp
--- a/Engine/Source/Developer/NewLogVisualizer/Private/STimelineBar.cpp
+++ b/Engine/Source/Developer/NewLogVisualizer/Private/STimelineBar.cpp
@@ -181,6 +181,7 @@ FVector2D STimelineBar::ComputeDesiredSize() const
 
 void STimelineBar::OnSelect()
 {
+	FSlateApplication::Get().SetKeyboardFocus(SharedThis(this), EFocusCause::Navigation);
 	CurrentItemIndex = INDEX_NONE;
 	SnapScrubPosition(TimeSliderController->GetTimeSliderArgs().ScrubPosition.Get());
 }
diff --git a/Engine/Source/Developer/NewLogVisualizer/Private/STimelinesContainer.cpp b/Engine/Source/Developer/NewLogVisualizer/Private/STimelinesContainer.cpp
--- a/Engine/Source/Developer/NewLogVisualizer/Private/STimelinesContainer.cpp
+++ b/Engine/Source/Developer/NewLogVisualizer/Private/STimelinesContainer.cpp
@@ -168,6 +168,45 @@ FReply STimelinesContainer::OnKeyDown(const FGeometry& MyGeometry, const FKeyEve
 		{
 			SelectedNodes.Reset();
 		}
+		return FReply::Handled();
+	}
+	else if (InKeyEvent.GetKey() == EKeys::Up || InKeyEvent.GetKey() == EKeys::Down)
+	{
+		TSharedPtr<class STimeline> PreviousTimeline;
+		TSharedPtr<class STimeline> LastSelected = SelectedNodes[SelectedNodes.Num() - 1];
+		for (int32 Index = 0; Index < TimelineItems.Num(); ++Index)
+		{
+			auto& CurrentItem = TimelineItems[Index];
+			if (LastSelected == CurrentItem)
+			{
+				if (InKeyEvent.GetKey() == EKeys::Up && PreviousTimeline.IsValid())
+				{
+					SetSelectionState(PreviousTimeline, true, true);
+				}
+				else if (InKeyEvent.GetKey() == EKeys::Down)
+				{
+					// let's find next visible time line
+					if (TimelineItems.IsValidIndex(Index + 1))
+					{
+						for (int32 i = Index + 1; i < TimelineItems.Num(); ++i)
+						{
+							if (TimelineItems[i]->GetVisibility() == EVisibility::Visible)
+							{
+								SetSelectionState(TimelineItems[i], true, true);
+								break;
+							}
+						}
+					}
+				}
+				break;
+			}
+
+			if (CurrentItem->GetVisibility() == EVisibility::Visible)
+			{
+				PreviousTimeline = CurrentItem;
+			}
+		}
+		return FReply::Handled();
 	}
 
 	return FReply::Unhandled();
