Fixes for camera toolbar button in LogVisualizertool
FIxed broken EVisualLoggerShapeElement::SinglePoint drawing from LogVisualizer tool

diff --git a/Engine/Source/Developer/NewLogVisualizer/Private/SVisualLogger.cpp b/Engine/Source/Developer/NewLogVisualizer/Private/SVisualLogger.cpp
--- a/Engine/Source/Developer/NewLogVisualizer/Private/SVisualLogger.cpp
+++ b/Engine/Source/Developer/NewLogVisualizer/Private/SVisualLogger.cpp
@@ -440,7 +440,7 @@ bool SVisualLogger::HandleCameraCommandIsChecked() const
 bool SVisualLogger::HandleCameraCommandCanExecute() const
 {
 	UWorld* World = VisualLoggerInterface->GetWorld();
-	return FVisualLogger::Get().IsRecording() && World && !World->bPlayersOnly && !World->bPlayersOnlyPending && World->IsPlayInEditor();
+	return FVisualLogger::Get().IsRecording() && World && (World->bPlayersOnly || World->bPlayersOnlyPending) && World->IsPlayInEditor() && (GEditor && !GEditor->bIsSimulatingInEditor);
 }
 
 void SVisualLogger::HandleCameraCommandExecute()
diff --git a/Engine/Source/Developer/NewLogVisualizer/Private/VisualLoggerCameraController.cpp b/Engine/Source/Developer/NewLogVisualizer/Private/VisualLoggerCameraController.cpp
--- a/Engine/Source/Developer/NewLogVisualizer/Private/VisualLoggerCameraController.cpp
+++ b/Engine/Source/Developer/NewLogVisualizer/Private/VisualLoggerCameraController.cpp
@@ -113,12 +113,13 @@ void AVisualLoggerCameraController::DisableCamera(UWorld* InWorld)
 {
 	check(InWorld);
 	for (FConstPlayerControllerIterator Iterator = InWorld->GetPlayerControllerIterator(); Iterator; ++Iterator)
-{
+	{
 		AVisualLoggerCameraController* VLogCam = Cast<AVisualLoggerCameraController>(*Iterator);
 		if (VLogCam && VLogCam->OriginalPlayer)
-	{
-		VLogCam->OriginalPlayer->SwitchController(VLogCam->OriginalControllerRef);
-		VLogCam->OnDeactivate(VLogCam->OriginalControllerRef);
+		{
+			VLogCam->OriginalPlayer->SwitchController(VLogCam->OriginalControllerRef);
+			VLogCam->OnDeactivate(VLogCam->OriginalControllerRef);
+			InWorld->DestroyActor(VLogCam, false, false);
 			break;
 		}
 	}
diff --git a/Engine/Source/Developer/NewLogVisualizer/Private/VisualLoggerRenderingActor.cpp b/Engine/Source/Developer/NewLogVisualizer/Private/VisualLoggerRenderingActor.cpp
--- a/Engine/Source/Developer/NewLogVisualizer/Private/VisualLoggerRenderingActor.cpp
+++ b/Engine/Source/Developer/NewLogVisualizer/Private/VisualLoggerRenderingActor.cpp
@@ -53,14 +53,7 @@ FPrimitiveSceneProxy* UVisualLoggerRenderingComponent::CreateSceneProxy()
 	}
 
 	FVisualLoggerSceneProxy *SceneProxy = new FVisualLoggerSceneProxy(this);
-	for (FVector& Point : RenderingActor->Points)
-	{
-		FDebugRenderSceneProxy::FSphere Sphere;
-		Sphere.Location = Point;
-		Sphere.Radius = 15;
-		Sphere.Color = FColor::White;
-		SceneProxy->Spheres.Add(Sphere);
-	}
+	SceneProxy->Spheres = RenderingActor->Points;
 	SceneProxy->Lines = RenderingActor->Lines;
 	SceneProxy->Cones = RenderingActor->Cones;
 	SceneProxy->Boxes = RenderingActor->Boxes;
@@ -154,7 +147,7 @@ void AVisualLoggerRenderingActor::AddDebugRendering()
 	}
 	{
 		const float Radius = 50;
-		Points.Add(FVector(1300,0,128));
+		Points.Add(FDebugRenderSceneProxy::FSphere(10, FVector(1300, 0, 128), FColor::White));
 	}
 }
 
@@ -250,7 +243,20 @@ void AVisualLoggerRenderingActor::OnItemSelectionChanged(const FVisualLogDevice:
 		{
 			const float Radius = float(ElementToDraw->Radius);
 			const bool bDrawLabel = ElementToDraw->Description.IsEmpty() == false;
-			Points.Append(ElementToDraw->Points);
+			for (int32 Index=0; Index < ElementToDraw->Points.Num(); ++Index)
+			{
+				const FVector& Point = ElementToDraw->Points[Index];
+				FDebugRenderSceneProxy::FSphere Sphere;
+				Sphere.Location = Point;
+				Sphere.Radius = Radius;
+				Sphere.Color = Color;
+				Points.Add(Sphere);
+				if (bDrawLabel)
+				{
+					const FString PrintString = FString::Printf(TEXT("%s_%d"), *ElementToDraw->Description, Index);
+					Texts.Add(FDebugRenderSceneProxy::FText3d(PrintString, Point, Color));
+				}
+			}
 		}
 		break;
 		case EVisualLoggerShapeElement::Polygon:
diff --git a/Engine/Source/Developer/NewLogVisualizer/Private/VisualLoggerRenderingActor.h b/Engine/Source/Developer/NewLogVisualizer/Private/VisualLoggerRenderingActor.h
--- a/Engine/Source/Developer/NewLogVisualizer/Private/VisualLoggerRenderingActor.h
+++ b/Engine/Source/Developer/NewLogVisualizer/Private/VisualLoggerRenderingActor.h
@@ -28,7 +28,7 @@ public:
 	TArray<FDebugRenderSceneProxy::FDebugLine> Lines;
 	TArray<FDebugRenderSceneProxy::FCone> Cones;
 	TArray<FDebugRenderSceneProxy::FDebugBox> Boxes;
-	TArray<FVector> Points;
+	TArray<FDebugRenderSceneProxy::FSphere> Points;
 	TArray<FDebugRenderSceneProxy::FMesh> Meshes;
 	TArray<FDebugRenderSceneProxy::FText3d> Texts;
 	TArray<FDebugRenderSceneProxy::FWireCylinder> Cylinders;
