Fixed text description draws as part of world/level from LogVisualizer

diff --git a/Engine/Source/Developer/NewLogVisualizer/Private/VisualLoggerRenderingActor.cpp b/Engine/Source/Developer/NewLogVisualizer/Private/VisualLoggerRenderingActor.cpp
--- a/Engine/Source/Developer/NewLogVisualizer/Private/VisualLoggerRenderingActor.cpp
+++ b/Engine/Source/Developer/NewLogVisualizer/Private/VisualLoggerRenderingActor.cpp
@@ -21,6 +21,8 @@ public:
 		: FDebugRenderSceneProxy(InComponent)
 	{
 		DrawType = SolidAndWireMeshes;
+		ViewFlagName = TEXT("VisLog");
+		ViewFlagIndex = uint32(FEngineShowFlags::FindIndexByName(*ViewFlagName));
 	}
 
 	virtual FPrimitiveViewRelevance GetViewRelevance(const FSceneView* View)
@@ -52,18 +54,18 @@ FPrimitiveSceneProxy* UVisualLoggerRenderingComponent::CreateSceneProxy()
 		return NULL;
 	}
 
-	FVisualLoggerSceneProxy *SceneProxy = new FVisualLoggerSceneProxy(this);
-	SceneProxy->Spheres = RenderingActor->Points;
-	SceneProxy->Lines = RenderingActor->Lines;
-	SceneProxy->Cones = RenderingActor->Cones;
-	SceneProxy->Boxes = RenderingActor->Boxes;
-	SceneProxy->Meshes = RenderingActor->Meshes;
-	SceneProxy->Cones = RenderingActor->Cones;
-	SceneProxy->Texts = RenderingActor->Texts;
-	SceneProxy->Cylinders = RenderingActor->Cylinders;
-	SceneProxy->Capsles = RenderingActor->Capsles;
-
-	return SceneProxy;
+	FVisualLoggerSceneProxy *VLogSceneProxy = new FVisualLoggerSceneProxy(this);
+	VLogSceneProxy->Spheres = RenderingActor->Points;
+	VLogSceneProxy->Lines = RenderingActor->Lines;
+	VLogSceneProxy->Cones = RenderingActor->Cones;
+	VLogSceneProxy->Boxes = RenderingActor->Boxes;
+	VLogSceneProxy->Meshes = RenderingActor->Meshes;
+	VLogSceneProxy->Cones = RenderingActor->Cones;
+	VLogSceneProxy->Texts = RenderingActor->Texts;
+	VLogSceneProxy->Cylinders = RenderingActor->Cylinders;
+	VLogSceneProxy->Capsles = RenderingActor->Capsles;
+
+	return VLogSceneProxy;
 }
 
 FBoxSphereBounds UVisualLoggerRenderingComponent::CalcBounds(const FTransform& LocalToWorld) const
@@ -246,11 +248,7 @@ void AVisualLoggerRenderingActor::OnItemSelectionChanged(const FVisualLogDevice:
 			for (int32 Index=0; Index < ElementToDraw->Points.Num(); ++Index)
 			{
 				const FVector& Point = ElementToDraw->Points[Index];
-				FDebugRenderSceneProxy::FSphere Sphere;
-				Sphere.Location = Point;
-				Sphere.Radius = Radius;
-				Sphere.Color = Color;
-				Points.Add(Sphere);
+				Points.Add(FDebugRenderSceneProxy::FSphere(Radius, Point, Color));
 				if (bDrawLabel)
 				{
 					const FString PrintString = FString::Printf(TEXT("%s_%d"), *ElementToDraw->Description, Index);
