Fixed cleanup after closing LogVisualizer tab

diff --git a/Engine/Source/Developer/NewLogVisualizer/Private/LogVisualizerModule.cpp b/Engine/Source/Developer/NewLogVisualizer/Private/LogVisualizerModule.cpp
--- a/Engine/Source/Developer/NewLogVisualizer/Private/LogVisualizerModule.cpp
+++ b/Engine/Source/Developer/NewLogVisualizer/Private/LogVisualizerModule.cpp
@@ -4,6 +4,7 @@
 #include "LogVisualizerModule.h"
 #include "LogVisualizerStyle.h"
 #include "SDockTab.h"
+#include "VisualLoggerRenderingActor.h"
 
 #define LOCTEXT_NAMESPACE "FLogVisualizerModule"
 
@@ -35,7 +36,7 @@ void FNewLogVisualizerModule::ShutdownModule()
 TSharedRef<SDockTab> FNewLogVisualizerModule::SpawnLogVisualizerTab(const FSpawnTabArgs& SpawnTabArgs)
 {
 	const TSharedRef<SDockTab> MajorTab = SNew(SDockTab)
-		.TabRole(ETabRole::MajorTab);
+		.TabRole(ETabRole::MajorTab).OnTabClosed(SDockTab::FOnTabClosedCallback::CreateRaw(this, &FNewLogVisualizerModule::OnTabClosed));
 
 	TSharedPtr<SWidget> TabContent;
 
@@ -46,5 +47,38 @@ TSharedRef<SDockTab> FNewLogVisualizerModule::SpawnLogVisualizerTab(const FSpawn
 	return MajorTab;
 }
 
+void FNewLogVisualizerModule::OnTabClosed(TSharedRef<SDockTab>)
+{
+	UWorld* World = NULL;
+#if WITH_EDITOR
+	UEditorEngine *EEngine = Cast<UEditorEngine>(GEngine);
+	if (GIsEditor && EEngine != NULL)
+	{
+		// lets use PlayWorld during PIE/Simulate and regular world from editor otherwise, to draw debug information
+		World = EEngine->PlayWorld != NULL ? EEngine->PlayWorld : EEngine->GetEditorWorldContext().World();
+
+	}
+	else 
+#endif
+	if (!GIsEditor)
+	{
+
+		World = GEngine->GetWorld();
+	}
+
+	if (World == NULL)
+	{
+		World = GWorld;
+	}
+
+	if (World)
+	{
+		for (TActorIterator<AVisualLoggerRenderingActor> It(World); It; ++It)
+		{
+			World->DestroyActor(*It);
+		}
+	}
+}
+
 IMPLEMENT_MODULE(FNewLogVisualizerModule, LogVisualizer);
 #undef LOCTEXT_NAMESPACE
diff --git a/Engine/Source/Developer/NewLogVisualizer/Public/LogVisualizerModule.h b/Engine/Source/Developer/NewLogVisualizer/Public/LogVisualizerModule.h
--- a/Engine/Source/Developer/NewLogVisualizer/Public/LogVisualizerModule.h
+++ b/Engine/Source/Developer/NewLogVisualizer/Public/LogVisualizerModule.h
@@ -18,4 +18,5 @@ public:
 
 private:
 	TSharedRef<SDockTab> SpawnLogVisualizerTab(const FSpawnTabArgs& SpawnTabArgs);
+	void OnTabClosed(TSharedRef<SDockTab>);
 };
