Fix Linux editor build.

diff --git a/Engine/Source/Developer/NewLogVisualizer/Private/LogVisualizerModule.cpp b/Engine/Source/Developer/NewLogVisualizer/Private/LogVisualizerModule.cpp
--- a/Engine/Source/Developer/NewLogVisualizer/Private/LogVisualizerModule.cpp
+++ b/Engine/Source/Developer/NewLogVisualizer/Private/LogVisualizerModule.cpp
@@ -5,6 +5,7 @@
 #include "LogVisualizerStyle.h"
 #include "SDockTab.h"
 #include "VisualLoggerRenderingActor.h"
+#include "Editor/EditorEngine.h"
 
 #define LOCTEXT_NAMESPACE "FLogVisualizerModule"
 
