Fixed CIS

diff --git a/Engine/Source/Developer/NewLogVisualizer/Private/SVisualLogger.cpp b/Engine/Source/Developer/NewLogVisualizer/Private/SVisualLogger.cpp
--- a/Engine/Source/Developer/NewLogVisualizer/Private/SVisualLogger.cpp
+++ b/Engine/Source/Developer/NewLogVisualizer/Private/SVisualLogger.cpp
@@ -17,6 +17,7 @@
 #	include "Editor/UnrealEd/Public/Editor.h"
 #	include "Editor/UnrealEd/Public/EditorViewportClient.h"
 #endif
+#include "ISettingsModule.h"
 
 #include "VisualLogger/VisualLoggerBinaryFileDevice.h"
 
@@ -262,6 +263,7 @@ void SVisualLogger::Construct(const FArguments& InArgs, const TSharedRef<SDockTa
 					)
 			)
 		);
+	//FGlobalTabmanager::Get()->RestoreFrom(Layout, TSharedPtr<SWindow>());
 
 	// Window Menu
 	FMenuBarBuilder MenuBarBuilder = FMenuBarBuilder(TSharedPtr<FUICommandList>());
