Fix for NonUnityTestCompile_OnMac build issue.

diff --git a/Engine/Source/Developer/NewLogVisualizer/Private/LogVisualizerSettings.cpp b/Engine/Source/Developer/NewLogVisualizer/Private/LogVisualizerSettings.cpp
--- a/Engine/Source/Developer/NewLogVisualizer/Private/LogVisualizerSettings.cpp
+++ b/Engine/Source/Developer/NewLogVisualizer/Private/LogVisualizerSettings.cpp
@@ -6,6 +6,7 @@
 #if WITH_EDITOR
 #include "Editor/EditorEngine.h"
 #include "ISettingsModule.h"
+#include "UnrealEdMisc.h"
 #endif // WITH_EDITOR
 
 ULogVisualizerSettings::ULogVisualizerSettings(const FObjectInitializer& ObjectInitializer) : Super(ObjectInitializer)
