UE-6063 "Fix compilation on case-sensitive systems" - GitHub #642 request

diff --git a/Engine/Source/Developer/NewLogVisualizer/Private/LogVisualizerModule.cpp b/Engine/Source/Developer/NewLogVisualizer/Private/LogVisualizerModule.cpp
--- a/Engine/Source/Developer/NewLogVisualizer/Private/LogVisualizerModule.cpp
+++ b/Engine/Source/Developer/NewLogVisualizer/Private/LogVisualizerModule.cpp
@@ -5,7 +5,7 @@
 #include "LogVisualizerStyle.h"
 #include "SDockTab.h"
 #include "VisualLoggerRenderingActor.h"
-#include "LogVIsualizerSettings.h"
+#include "LogVisualizerSettings.h"
 #if WITH_EDITOR
 #include "Editor/EditorEngine.h"
 #include "ISettingsModule.h"
diff --git a/Engine/Source/Developer/NewLogVisualizer/Private/LogVisualizerSettings.cpp b/Engine/Source/Developer/NewLogVisualizer/Private/LogVisualizerSettings.cpp
--- a/Engine/Source/Developer/NewLogVisualizer/Private/LogVisualizerSettings.cpp
+++ b/Engine/Source/Developer/NewLogVisualizer/Private/LogVisualizerSettings.cpp
@@ -2,7 +2,7 @@
 
 #include "LogVisualizer.h"
 #include "Misc/CoreMisc.h"
-#include "LogVIsualizerSettings.h"
+#include "LogVisualizerSettings.h"
 #if WITH_EDITOR
 #include "Editor/EditorEngine.h"
 #include "ISettingsModule.h"
