Increased default graph data range in LogVisualizer tool.

diff --git a/Engine/Source/Developer/NewLogVisualizer/Private/TimeSliderController.h b/Engine/Source/Developer/NewLogVisualizer/Private/TimeSliderController.h
--- a/Engine/Source/Developer/NewLogVisualizer/Private/TimeSliderController.h
+++ b/Engine/Source/Developer/NewLogVisualizer/Private/TimeSliderController.h
@@ -7,7 +7,7 @@
 struct FVisualLoggerTimeSliderArgs : FTimeSliderArgs
 {
 	FVisualLoggerTimeSliderArgs()
-		: CursorSize(0.0f)
+		: CursorSize(0.1f)
 	{}
 
 	/** Cursor range for data like histogram graphs, etc. */
