Added double click reaction to LogVisualizer, to setup camera location over current log location.

diff --git a/Engine/Source/Developer/NewLogVisualizer/Private/LogVisualizer.h b/Engine/Source/Developer/NewLogVisualizer/Private/LogVisualizer.h
--- a/Engine/Source/Developer/NewLogVisualizer/Private/LogVisualizer.h
+++ b/Engine/Source/Developer/NewLogVisualizer/Private/LogVisualizer.h
@@ -60,3 +60,4 @@ protected:
 #include "SVisualLoggerLogsList.h"
 #include "SVisualLoggerStatusView.h"
 #include "STimeline.h"
+#include "LogVisualizerSettings.h"
\ No newline at end of file
diff --git a/Engine/Source/Developer/NewLogVisualizer/Private/LogVisualizerSettings.cpp b/Engine/Source/Developer/NewLogVisualizer/Private/LogVisualizerSettings.cpp
--- a/Engine/Source/Developer/NewLogVisualizer/Private/LogVisualizerSettings.cpp
+++ b/Engine/Source/Developer/NewLogVisualizer/Private/LogVisualizerSettings.cpp
@@ -12,6 +12,7 @@
 ULogVisualizerSettings::ULogVisualizerSettings(const FObjectInitializer& ObjectInitializer) : Super(ObjectInitializer)
 {
 	TrivialLogsThreshold = 1;
+	DefaultCameraDistance = 150;
 }
 
 void ULogVisualizerSettings::PostEditChangeProperty(struct FPropertyChangedEvent& PropertyChangedEvent)
diff --git a/Engine/Source/Developer/NewLogVisualizer/Private/STimelineBar.cpp b/Engine/Source/Developer/NewLogVisualizer/Private/STimelineBar.cpp
--- a/Engine/Source/Developer/NewLogVisualizer/Private/STimelineBar.cpp
+++ b/Engine/Source/Developer/NewLogVisualizer/Private/STimelineBar.cpp
@@ -2,6 +2,16 @@
 #include "STimelineBar.h"
 #include "TimeSliderController.h"
 #include "STimelinesContainer.h"
+#include "VisualLoggerCameraController.h"
+#if WITH_EDITOR
+#	include "Editor/UnrealEd/Public/EditorComponents.h"
+#	include "Editor/UnrealEd/Public/EditorReimportHandler.h"
+#	include "Editor/UnrealEd/Public/TexAlignTools.h"
+#	include "Editor/UnrealEd/Public/TickableEditorObject.h"
+#	include "UnrealEdClasses.h"
+#	include "Editor/UnrealEd/Public/Editor.h"
+#	include "Editor/UnrealEd/Public/EditorViewportClient.h"
+#endif
 
 FReply STimelineBar::OnMouseButtonDown(const FGeometry& MyGeometry, const FPointerEvent& MouseEvent)
 {
@@ -33,6 +43,76 @@ FReply STimelineBar::OnMouseMove(const FGeometry& MyGeometry, const FPointerEven
 	return TimeSliderController->OnMouseMove(SharedThis(this), MyGeometry, MouseEvent);
 }
 
+FReply STimelineBar::OnMouseButtonDoubleClick(const FGeometry& MyGeometry, const FPointerEvent& MouseEvent)
+{
+	auto &Entries = TimelineOwner.Pin()->GetEntries();
+	SnapScrubPosition(TimeSliderController->GetTimeSliderArgs().ScrubPosition.Get());
+	UWorld* World = VisualLoggerInterface->GetWorld();
+	if (World && MouseEvent.GetEffectingButton() == EKeys::LeftMouseButton && Entries.IsValidIndex(CurrentItemIndex))
+	{
+		FVector CurrentLocation = Entries[CurrentItemIndex].Entry.Location;
+
+		FVector Extent(150);
+		bool bFoundActor = false;
+		FName OwnerName = Entries[CurrentItemIndex].OwnerName;
+		for (FActorIterator It(World); It; ++It)
+		{
+			AActor* Actor = *It;
+			if (Actor->GetFName() == OwnerName)
+			{
+				FVector Orgin;
+				Actor->GetActorBounds(false, Orgin, Extent);
+				bFoundActor = true;
+				break;
+			}
+		}
+
+		
+		const float DefaultCameraDistance = ULogVisualizerSettings::StaticClass()->GetDefaultObject<ULogVisualizerSettings>()->DefaultCameraDistance;
+		Extent = Extent.Size() < DefaultCameraDistance ? FVector(1) * DefaultCameraDistance : Extent;
+
+#if WITH_EDITOR
+		UEditorEngine *EEngine = Cast<UEditorEngine>(GEngine);
+		if (GIsEditor && EEngine != NULL)
+		{
+			for (auto ViewportClient : EEngine->AllViewportClients)
+			{
+				ViewportClient->FocusViewportOnBox(FBox::BuildAABB(CurrentLocation, Extent));
+			}
+		}
+		else if (AVisualLoggerCameraController::IsEnabled(World) && AVisualLoggerCameraController::Instance.IsValid() && AVisualLoggerCameraController::Instance->GetSpectatorPawn())
+		{
+			ULocalPlayer* LocalPlayer = Cast<ULocalPlayer>(AVisualLoggerCameraController::Instance->Player);
+			if (LocalPlayer && LocalPlayer->ViewportClient && LocalPlayer->ViewportClient->Viewport)
+			{
+				FViewport* Viewport = LocalPlayer->ViewportClient->Viewport;
+
+				FBox BoundingBox = FBox::BuildAABB(CurrentLocation, Extent);
+				const FVector Position = BoundingBox.GetCenter();
+				float Radius = BoundingBox.GetExtent().Size();
+
+				FViewportCameraTransform ViewTransform;
+				ViewTransform.TransitionToLocation(Position, true);
+
+				float NewOrthoZoom;
+				const float AspectRatio = 1.777777f;
+				uint32 MinAxisSize = (AspectRatio > 1.0f) ? Viewport->GetSizeXY().Y : Viewport->GetSizeXY().X;
+				float Zoom = Radius / (MinAxisSize / 2);
+
+				NewOrthoZoom = Zoom * (Viewport->GetSizeXY().X*15.0f);
+				NewOrthoZoom = FMath::Clamp<float>(NewOrthoZoom, 250, MAX_FLT);
+				ViewTransform.SetOrthoZoom(NewOrthoZoom);
+
+				AVisualLoggerCameraController::Instance->GetSpectatorPawn()->TeleportTo(ViewTransform.GetLocation(), ViewTransform.GetRotation(), false, true);
+			}
+		}
+#endif
+		return FReply::Handled();
+	}
+
+	return FReply::Unhandled();
+}
+
 FReply STimelineBar::OnKeyDown(const FGeometry& MyGeometry, const FKeyEvent& InKeyEvent)
 {
 	FReply ReturnValue = FReply::Unhandled();
diff --git a/Engine/Source/Developer/NewLogVisualizer/Private/STimelineBar.h b/Engine/Source/Developer/NewLogVisualizer/Private/STimelineBar.h
--- a/Engine/Source/Developer/NewLogVisualizer/Private/STimelineBar.h
+++ b/Engine/Source/Developer/NewLogVisualizer/Private/STimelineBar.h
@@ -15,6 +15,7 @@ public:
 	SLATE_END_ARGS()
 
 	/** Constructs this widget with InArgs */
+	virtual FReply OnMouseButtonDoubleClick(const FGeometry& InMyGeometry, const FPointerEvent& InMouseEvent) override;
 	virtual FReply OnMouseButtonDown(const FGeometry& MyGeometry, const FPointerEvent& MouseEvent) override;
 	virtual FReply OnMouseButtonUp(const FGeometry& MyGeometry, const FPointerEvent& MouseEvent) override;
 	virtual FReply OnMouseMove(const FGeometry& MyGeometry, const FPointerEvent& MouseEvent) override;
diff --git a/Engine/Source/Developer/NewLogVisualizer/Private/SVisualLogger.cpp b/Engine/Source/Developer/NewLogVisualizer/Private/SVisualLogger.cpp
--- a/Engine/Source/Developer/NewLogVisualizer/Private/SVisualLogger.cpp
+++ b/Engine/Source/Developer/NewLogVisualizer/Private/SVisualLogger.cpp
@@ -669,6 +669,7 @@ void SVisualLogger::OnObjectSelectionChanged(TSharedPtr<class STimeline> TimeLin
 	{
 		HelperActor->ObjectSelectionChanged(TimeLine);
 	}
+	MainView->OnObjectSelectionChanged(TimeLine);
 }
 
 void SVisualLogger::OnFiltersSearchChanged(const FText& Filter)
diff --git a/Engine/Source/Developer/NewLogVisualizer/Private/SVisualLoggerView.cpp b/Engine/Source/Developer/NewLogVisualizer/Private/SVisualLoggerView.cpp
--- a/Engine/Source/Developer/NewLogVisualizer/Private/SVisualLoggerView.cpp
+++ b/Engine/Source/Developer/NewLogVisualizer/Private/SVisualLoggerView.cpp
@@ -33,7 +33,6 @@ void SVisualLoggerView::Construct(const FArguments& InArgs, const TSharedRef<FUI
 	TSharedRef<ITimeSlider> TopTimeSlider = SNew(STimeSlider, TimeSliderController.ToSharedRef()).MirrorLabels(bMirrorLabels);
 	TSharedRef<ITimeSlider> BottomTimeSlider = SNew(STimeSlider, TimeSliderController.ToSharedRef()).MirrorLabels(bMirrorLabels);
 
-	TSharedPtr<SScrollBox> ScrollBox;
 	TSharedRef<SScrollBar> ScrollBar =
 		SNew(SScrollBar)
 		.Thickness(FVector2D(5.0f, 5.0f));
@@ -158,6 +157,17 @@ void SVisualLoggerView::Construct(const FArguments& InArgs, const TSharedRef<FUI
 
 }
 
+void SVisualLoggerView::OnObjectSelectionChanged(TSharedPtr<class STimeline> TimeLine)
+{
+	//FIXME: scroll to selected timeline (SebaK)
+	//FWidgetPath WidgetPath;
+	//if (FSlateApplication::Get().GeneratePathToWidgetUnchecked(TimeLine.ToSharedRef(), WidgetPath))
+	//{
+	//	FArrangedWidget ArrangedWidget = WidgetPath.FindArrangedWidget(TimeLine.ToSharedRef()).Get(FArrangedWidget::NullWidget);
+	//	ScrollBox->ScrollDescendantIntoView(ArrangedWidget.Geometry, TimeLine, true);
+	//}
+}
+
 void SVisualLoggerView::OnSearchSplitterResized()
 {
 	SSplitter::FSlot const& LeftSplitterSlot = SearchSplitter->SlotAt(0);
diff --git a/Engine/Source/Developer/NewLogVisualizer/Private/SVisualLoggerView.h b/Engine/Source/Developer/NewLogVisualizer/Private/SVisualLoggerView.h
--- a/Engine/Source/Developer/NewLogVisualizer/Private/SVisualLoggerView.h
+++ b/Engine/Source/Developer/NewLogVisualizer/Private/SVisualLoggerView.h
@@ -31,12 +31,14 @@ public:
 	void OnFiltersChanged();
 	void OnSearchChanged(const FText& Filter);
 	void OnSearchSplitterResized();
+	void OnObjectSelectionChanged(TSharedPtr<class STimeline> TimeLine);
 
 	void GetTimelines(TArray<TSharedPtr<class STimeline> >&, bool bOnlySelectedOnes = false);
 
 protected:
 	TSharedPtr<class STimelinesContainer> TimelinesContainer;
 	TSharedPtr<class SSplitter> SearchSplitter;
+	TSharedPtr<class SScrollBox> ScrollBox;
 
 	FVisualLoggerEvents	VisualLoggerEvents;
 	float AnimationOutlinerFillPercentage;
diff --git a/Engine/Source/Developer/NewLogVisualizer/Public/LogVisualizerSettings.h b/Engine/Source/Developer/NewLogVisualizer/Public/LogVisualizerSettings.h
--- a/Engine/Source/Developer/NewLogVisualizer/Public/LogVisualizerSettings.h
+++ b/Engine/Source/Developer/NewLogVisualizer/Public/LogVisualizerSettings.h
@@ -37,6 +37,10 @@ public:
 	UPROPERTY(EditAnywhere, config, Category = "VisualLogger")
 	bool bShowHistogramLabelsOutside;
 
+	/** Camera distance used to setup location during reaction on log item double click */
+	UPROPERTY(EditAnywhere, config, Category = "VisualLogger", meta = (ClampMin = "10", ClampMax = "1000", UIMin = "10", UIMax = "1000"))
+	float DefaultCameraDistance;
+
 	DECLARE_EVENT_OneParam(ULogVisualizerSettings, FSettingChangedEvent, FName /*PropertyName*/);
 	FSettingChangedEvent& OnSettingChanged() { return SettingChangedEvent; }
 
diff --git a/Engine/Source/Runtime/Engine/Private/VisualLogger/VisualLogger.cpp b/Engine/Source/Runtime/Engine/Private/VisualLogger/VisualLogger.cpp
--- a/Engine/Source/Runtime/Engine/Private/VisualLogger/VisualLogger.cpp
+++ b/Engine/Source/Runtime/Engine/Private/VisualLogger/VisualLogger.cpp
@@ -38,7 +38,7 @@ FVisualLogger::FVisualLogger()
 
 UWorld* FVisualLogger::GetWorld()
 {
-	UWorld* World = NULL;
+	UWorld* World = GWorld;
 #if WITH_EDITOR
 	UEditorEngine *EEngine = Cast<UEditorEngine>(GEngine);
 	if (GIsEditor && EEngine != NULL)
