Added reports generated from visual logger events to LogVisualizer.
Removing all data from LogVisualizer cleanups rendering data correctly.
UE4 Editor loads LogVisualizer module - to have it early enough in editor.

diff --git a/Engine/Source/Developer/NewLogVisualizer/Private/STimeline.cpp b/Engine/Source/Developer/NewLogVisualizer/Private/STimeline.cpp
--- a/Engine/Source/Developer/NewLogVisualizer/Private/STimeline.cpp
+++ b/Engine/Source/Developer/NewLogVisualizer/Private/STimeline.cpp
@@ -13,7 +13,12 @@ STimeline::~STimeline()
 FReply STimeline::OnMouseButtonDown(const FGeometry& MyGeometry, const FPointerEvent& MouseEvent)
 {
 	Owner->ChangeSelection(SharedThis(this), MouseEvent);
-	return FReply::Unhandled();
+	if (OnGetMenuContent.IsBound() && MouseEvent.GetEffectingButton() == EKeys::RightMouseButton)
+	{
+		PopupAnchor->SetIsOpen(!PopupAnchor->IsOpen());
+	}
+
+	return FReply::Handled();
 }
 
 FReply STimeline::OnMouseButtonUp(const FGeometry& MyGeometry, const FPointerEvent& MouseEvent)
@@ -90,6 +95,8 @@ void STimeline::HandleLogVisualizerSettingChanged(FName Name)
 void STimeline::Construct(const FArguments& InArgs, TSharedPtr<SVisualLoggerView> VisualLoggerView, TSharedPtr<FSequencerTimeSliderController> TimeSliderController, TSharedPtr<STimelinesContainer> InContainer, const FVisualLogDevice::FVisualLogEntryItem& Entry)
 {
 	VisualLoggerInterface = InArgs._VisualLoggerInterface.Get();
+	OnGetMenuContent = InArgs._OnGetMenuContent;
+
 	Owner = InContainer;
 	Name = Entry.OwnerName;
 
@@ -108,20 +115,24 @@ void STimeline::Construct(const FArguments& InArgs, TSharedPtr<SVisualLoggerView
 			.VAlign(VAlign_Fill)
 			.FillWidth(TAttribute<float>(VisualLoggerView.Get(), &SVisualLoggerView::GetAnimationOutlinerFillPercentage))
 			[
-				SNew(SBorder)
-				.HAlign(HAlign_Fill)
-				.Padding(FMargin(0, 0, 4, 0))
-				.BorderImage(FCoreStyle::Get().GetBrush("NoBorder"))
+				SAssignNew(PopupAnchor, SMenuAnchor)
+				.OnGetMenuContent(OnGetMenuContent)
 				[
 					SNew(SBorder)
-					.VAlign(VAlign_Center)
-					.BorderImage(this, &STimeline::GetBorder)
-					.Padding(FMargin(4, 0, 2, 0))
+					.HAlign(HAlign_Fill)
+					.Padding(FMargin(0, 0, 4, 0))
+					.BorderImage(FCoreStyle::Get().GetBrush("NoBorder"))
 					[
-						// Search box for searching through the outliner
-						SNew(STextBlock)
-						.Text(FString(Name.ToString()))
-						.ShadowOffset(FVector2D(1.f, 1.f))
+						SNew(SBorder)
+						.VAlign(VAlign_Center)
+						.BorderImage(this, &STimeline::GetBorder)
+						.Padding(FMargin(4, 0, 2, 0))
+						[
+							// Search box for searching through the outliner
+							SNew(STextBlock)
+							.Text(FString(Name.ToString()))
+							.ShadowOffset(FVector2D(1.f, 1.f))
+						]
 					]
 				]
 			]
diff --git a/Engine/Source/Developer/NewLogVisualizer/Private/STimeline.h b/Engine/Source/Developer/NewLogVisualizer/Private/STimeline.h
--- a/Engine/Source/Developer/NewLogVisualizer/Private/STimeline.h
+++ b/Engine/Source/Developer/NewLogVisualizer/Private/STimeline.h
@@ -9,8 +9,9 @@ class STimeline : public SCompoundWidget
 {
 public:
 	SLATE_BEGIN_ARGS(STimeline){}
-		SLATE_EVENT(FOnItemSelectionChanged, OnItemSelectionChanged)
 		SLATE_ATTRIBUTE(TSharedPtr<IVisualLoggerInterface>, VisualLoggerInterface)
+		SLATE_EVENT(FOnItemSelectionChanged, OnItemSelectionChanged)
+		SLATE_EVENT(FOnGetContent, OnGetMenuContent)
 	SLATE_END_ARGS();
 
 	virtual FReply OnMouseButtonDown(const FGeometry& MyGeometry, const FPointerEvent& MouseEvent) override;
@@ -40,7 +41,11 @@ protected:
 	TSharedPtr<class STimelinesContainer> Owner;
 	TSharedPtr<class STimelineBar> TimelineBar;
 	TSharedPtr<IVisualLoggerInterface> VisualLoggerInterface;
+	TSharedPtr<SMenuAnchor> PopupAnchor;
 
 	FName Name;
 	FString SearchFilter;
+
+	/** Delegate to execute to get the menu content of this timeline */
+	FOnGetContent OnGetMenuContent;
 };
diff --git a/Engine/Source/Developer/NewLogVisualizer/Private/STimelinesContainer.cpp b/Engine/Source/Developer/NewLogVisualizer/Private/STimelinesContainer.cpp
--- a/Engine/Source/Developer/NewLogVisualizer/Private/STimelinesContainer.cpp
+++ b/Engine/Source/Developer/NewLogVisualizer/Private/STimelinesContainer.cpp
@@ -3,6 +3,9 @@
 #include "STimeline.h"
 #include "STimelineBar.h"
 #include "TimeSliderController.h"
+#include "SVisualLoggerReport.h"
+
+#define LOCTEXT_NAMESPACE "STimelinesContainer"
 
 TSharedRef<SWidget> STimelinesContainer::MakeTimeline(TSharedPtr<class SVisualLoggerView> VisualLoggerView, TSharedPtr<class FSequencerTimeSliderController> TimeSliderController, const FVisualLogDevice::FVisualLogEntryItem& Entry)
 {
@@ -13,6 +16,7 @@ TSharedRef<SWidget> STimelinesContainer::MakeTimeline(TSharedPtr<class SVisualLo
 			SAssignNew(NewTimeline, STimeline, VisualLoggerView, TimeSliderController, SharedThis(this), Entry)
 			.OnItemSelectionChanged(this->VisualLoggerInterface->GetVisualLoggerEvents().OnItemSelectionChanged)
 			.VisualLoggerInterface(this->VisualLoggerInterface)
+			.OnGetMenuContent(this, &STimelinesContainer::GetRightClickMenuContent)
 		];
 
 	TimelineItems.Add(NewTimeline.ToSharedRef());
@@ -20,6 +24,37 @@ TSharedRef<SWidget> STimelinesContainer::MakeTimeline(TSharedPtr<class SVisualLo
 	return NewTimeline.ToSharedRef();
 }
 
+TSharedRef<SWidget> STimelinesContainer::GetRightClickMenuContent()
+{
+	FMenuBuilder MenuBuilder(true, NULL);
+	MenuBuilder.BeginSection("VisualLogReports", LOCTEXT("VisualLogReports", "VisualLog Reports"));
+	{
+		MenuBuilder.AddMenuEntry(
+			LOCTEXT("GenarateReport", "Genarate Report"),
+			LOCTEXT("GenarateReportTooltip", "Genarate report from Visual Log events."),
+			FSlateIcon(),
+			FUIAction(FExecuteAction::CreateSP(this, &STimelinesContainer::GenerateReport))
+			);
+	}
+	MenuBuilder.EndSection();
+
+	FDisplayMetrics DisplayMetrics;
+	FSlateApplication::Get().GetDisplayMetrics(DisplayMetrics);
+
+	const FVector2D DisplaySize(
+		DisplayMetrics.PrimaryDisplayWorkAreaRect.Right - DisplayMetrics.PrimaryDisplayWorkAreaRect.Left,
+		DisplayMetrics.PrimaryDisplayWorkAreaRect.Bottom - DisplayMetrics.PrimaryDisplayWorkAreaRect.Top);
+
+	return
+		SNew(SVerticalBox)
+
+		+ SVerticalBox::Slot()
+		.MaxHeight(DisplaySize.Y * 0.5)
+		[
+			MenuBuilder.MakeWidget()
+		];
+}
+
 void STimelinesContainer::SetSelectionState(TSharedPtr<class STimeline> AffectedNode, bool bSelect, bool bDeselectOtherNodes)
 {
 	if (bSelect)
@@ -30,11 +65,14 @@ void STimelinesContainer::SetSelectionState(TSharedPtr<class STimeline> Affected
 			SelectedNodes.Empty();
 		}
 
-		SelectedNodes.Add(AffectedNode);
-		AffectedNode->OnSelect();
+		if (AffectedNode.IsValid())
+		{
+			SelectedNodes.Add(AffectedNode);
+			AffectedNode->OnSelect();
+		}
 		VisualLoggerInterface->GetVisualLoggerEvents().OnObjectSelectionChanged.ExecuteIfBound(AffectedNode);
 	}
-	else
+	else if (AffectedNode.IsValid())
 	{
 		// Not selecting so remove the node from the selection set
 		SelectedNodes.Remove(AffectedNode);
@@ -62,26 +100,34 @@ void STimelinesContainer::ChangeSelection(class TSharedPtr<class STimeline> Time
 	}
 	else
 	{
-		TSharedPtr<class STimeline> LastSelected = SelectedNodes[SelectedNodes.Num() - 1];
-		bool bStartedSelection = false;
-		for (auto& CurrentItem : TimelineItems)
+		if (SelectedNodes.Num() == 0 && TimelineItems.Num())
 		{
-			if (CurrentItem == LastSelected || CurrentItem == Timeline)
+			SetSelectionState(TimelineItems[0], true, true);
+		}
+
+		TSharedPtr<class STimeline> LastSelected = SelectedNodes.Num() ? SelectedNodes[SelectedNodes.Num() - 1] : nullptr;
+		if (LastSelected.IsValid())
+		{
+			bool bStartedSelection = false;
+			for (auto& CurrentItem : TimelineItems)
 			{
-				if (!bStartedSelection)
+				if (CurrentItem == LastSelected || CurrentItem == Timeline)
 				{
-					bStartedSelection = true;
+					if (!bStartedSelection)
+					{
+						bStartedSelection = true;
+					}
+					else
+					{
+						bStartedSelection = false;
+						break;
+					}
 				}
-				else
+				if (bStartedSelection)
 				{
-					bStartedSelection = false;
-					break;
+					SetSelectionState(CurrentItem, true, false);
 				}
 			}
-			if (bStartedSelection)
-			{
-				SetSelectionState(CurrentItem, true, false);
-			}
 		}
 		SetSelectionState(Timeline, true, false);
 	}
@@ -160,14 +206,7 @@ FReply STimelinesContainer::OnKeyDown(const FGeometry& MyGeometry, const FKeyEve
 			ContainingBorder->RemoveSlot(CurrentNode.ToSharedRef());
 		}
 
-		if (NotSelectedOne.IsValid())
-		{
-			SetSelectionState(NotSelectedOne.Pin(), true, true);
-		}
-		else
-		{
-			SelectedNodes.Reset();
-		}
+		SetSelectionState(NotSelectedOne.Pin(), true, true);
 		return FReply::Handled();
 	}
 	else if (InKeyEvent.GetKey() == EKeys::Up || InKeyEvent.GetKey() == EKeys::Down)
@@ -276,3 +315,31 @@ void STimelinesContainer::OnFiltersChanged()
 		CurrentItem->OnFiltersChanged();
 	}
 }
+
+void STimelinesContainer::GenerateReport()
+{
+	//TArray< TSharedPtr<FLogsListItem> > ItemsToSave = LogsListWidget->GetSelectedItems();
+	//TArray< TSharedPtr<class STimeline> > AllLogs;
+
+	//TSharedPtr<FLogsListItem>* LogListItem = ItemsToSave.GetData();
+	//for (int32 ItemIndex = 0; ItemIndex < ItemsToSave.Num(); ++ItemIndex, ++LogListItem)
+	//{
+	//	if (LogListItem->IsValid() && LogVisualizer->Logs.IsValidIndex((*LogListItem)->LogIndex))
+	//	{
+	//		TSharedPtr<FActorsVisLog> Log = LogVisualizer->Logs[(*LogListItem)->LogIndex];
+	//		AllLogs.Add(Log);
+	//	}
+	//}
+
+
+	TSharedRef<SWindow> NewWindow = SNew(SWindow)
+		.ClientSize(FVector2D(720, 768))
+		.Title(NSLOCTEXT("LogVisualizerReport", "WindowTitle", "Log Visualizer Report"))
+		[
+			SNew(SVisualLoggerReport, SelectedNodes)
+		];
+
+	FSlateApplication::Get().AddWindow(NewWindow);
+}
+
+#undef LOCTEXT_NAMESPACE
diff --git a/Engine/Source/Developer/NewLogVisualizer/Private/STimelinesContainer.h b/Engine/Source/Developer/NewLogVisualizer/Private/STimelinesContainer.h
--- a/Engine/Source/Developer/NewLogVisualizer/Private/STimelinesContainer.h
+++ b/Engine/Source/Developer/NewLogVisualizer/Private/STimelinesContainer.h
@@ -21,6 +21,7 @@ public:
 	/** Constructs this widget with InArgs */
 	void Construct(const FArguments& InArgs, TSharedRef<class SVisualLoggerView>, TSharedRef<FSequencerTimeSliderController> TimeSliderController);
 	TSharedRef<SWidget> MakeTimeline(TSharedPtr<class SVisualLoggerView> VisualLoggerView, TSharedPtr<class FSequencerTimeSliderController> TimeSliderController, const FVisualLogDevice::FVisualLogEntryItem& Entry);
+	TSharedRef<SWidget> GetRightClickMenuContent();
 
 	void OnTimelineSelected(TSharedPtr<class STimelinesBar> Widget);
 	void ChangeSelection(class TSharedPtr<class STimeline>, const FPointerEvent& MouseEvent);
@@ -29,6 +30,7 @@ public:
 	void OnNewLogEntry(const FVisualLogDevice::FVisualLogEntryItem& Entry);
 	void OnFiltersChanged();
 	void OnSearchChanged(const FText& Filter);
+	void GenerateReport();
 
 	/**
 	* Selects a node in the tree
diff --git a/Engine/Source/Developer/NewLogVisualizer/Private/SVisualLoggerReport.cpp b/Engine/Source/Developer/NewLogVisualizer/Private/SVisualLoggerReport.cpp
new file mode 100644
--- /dev/null
+++ b/Engine/Source/Developer/NewLogVisualizer/Private/SVisualLoggerReport.cpp
@@ -0,0 +1,160 @@
+// Copyright 1998-2014 Epic Games, Inc. All Rights Reserved.
+
+#include "LogVisualizer.h"
+#include "Debug/ReporterGraph.h"
+#include "MainFrame.h"
+#include "DesktopPlatformModule.h"
+#include "Editor/UnrealEd/Classes/Editor/EditorEngine.h"
+#include "Editor/UnrealEd/Public/EditorComponents.h"
+#include "Editor/UnrealEd/Public/EditorReimportHandler.h"
+#include "Editor/UnrealEd/Public/TexAlignTools.h"
+#include "Editor/UnrealEd/Public/TickableEditorObject.h"
+#include "UnrealEdClasses.h"
+#include "Editor/UnrealEd/Public/Editor.h"
+#include "Editor/UnrealEd/Public/EditorViewportClient.h"
+#include "VisualLoggerTypes.h"
+#include "SVisualLoggerReport.h"
+
+#define LOCTEXT_NAMESPACE "SVisualLoggerReport"
+
+BEGIN_SLATE_FUNCTION_BUILD_OPTIMIZATION
+void SVisualLoggerReport::Construct(const FArguments& InArgs, TArray< TSharedPtr<class STimeline> >& InSelectedItems)
+{
+	SelectedItems = InSelectedItems;
+	this->ChildSlot
+		[
+			SNew(SScrollBox)
+			+ SScrollBox::Slot()
+			[
+				SNew(SBorder)
+				.BorderImage(FEditorStyle::Get().GetBrush("ToolPanel.GroupBorder"))
+				[
+					SNew(SVerticalBox)
+					+ SVerticalBox::Slot().AutoHeight().Padding(0)
+					[
+						SNew(SBorder)
+						.BorderImage(FEditorStyle::Get().GetBrush("Menu.Background"))
+						.Padding(0)
+						[
+							SNew(SRichTextBlock)
+							.Text(this, &SVisualLoggerReport::GenerateReportText)
+							.TextStyle(FEditorStyle::Get(), "ContentBrowser.Filters.Text")
+							.DecoratorStyleSet(&FEditorStyle::Get())
+							.WrapTextAt(800)
+							.Justification(ETextJustify::Left)
+							.Margin(FMargin(20))
+						]
+					]
+				]
+			]
+		];
+}
+END_SLATE_FUNCTION_BUILD_OPTIMIZATION
+
+SVisualLoggerReport::~SVisualLoggerReport()
+{
+}
+
+FText SVisualLoggerReport::GenerateReportText() const
+{
+	FString OutString;
+	TArray<FVisualLogEvent> GlobalEventsStats;
+	TMap<FString, TArray<FString> >	 EventToObjectsMap;
+
+
+	for (TSharedPtr<class STimeline> LogItem : SelectedItems)
+	{
+		TArray<FVisualLogEvent> AllEvents;
+
+		for (const FVisualLogDevice::FVisualLogEntryItem& CurrentEntry : LogItem->GetEntries())
+		{
+			for (const FVisualLogEvent& Event : CurrentEntry.Entry.Events)
+			{
+				int32 Index = AllEvents.Find(FVisualLogEvent(Event));
+				if (Index != INDEX_NONE)
+				{
+					for (auto& CurrentEventTag : Event.EventTags)
+					{
+						if (AllEvents[Index].EventTags.Contains(CurrentEventTag.Key))
+						{
+							AllEvents[Index].EventTags[CurrentEventTag.Key] += CurrentEventTag.Value;
+						}
+						else
+						{
+							AllEvents[Index].EventTags.Add(CurrentEventTag.Key, CurrentEventTag.Value);
+						}
+					}
+					AllEvents[Index].Counter++;
+				}
+				else
+				{
+					AllEvents.Add(Event);
+				}
+			}
+		}
+
+		bool bPrintNextLine = false;
+		if (AllEvents.Num() > 0)
+		{
+			OutString.Append(FString::Printf(TEXT("%s\n"), *LogItem->GetName().ToString()));
+		}
+		for (auto& CurrentEvent : AllEvents)
+		{
+			for (auto& CurrentEventTag : CurrentEvent.EventTags)
+			{
+				OutString.Append(FString::Printf(TEXT("        %s  with <%s> tag occurred    %d times\n"), *CurrentEvent.Name, *CurrentEventTag.Key.ToString(), CurrentEventTag.Value));
+			}
+			OutString.Append(FString::Printf(TEXT("        %s occurred %d times\n\n"), *CurrentEvent.Name, CurrentEvent.Counter));
+
+			bool bJustAdded = false;
+			int32 Index = GlobalEventsStats.Find(FVisualLogEvent(CurrentEvent));
+			if (Index != INDEX_NONE)
+			{
+				GlobalEventsStats[Index].Counter += CurrentEvent.Counter;
+				EventToObjectsMap[CurrentEvent.Name].AddUnique(LogItem->GetName().ToString());
+			}
+			else
+			{
+				bJustAdded = true;
+				GlobalEventsStats.Add(CurrentEvent);
+				EventToObjectsMap.FindOrAdd(CurrentEvent.Name).Add(LogItem->GetName().ToString());
+			}
+
+			Index = GlobalEventsStats.Find(FVisualLogEvent(CurrentEvent));
+			for (auto& CurrentEventTag : CurrentEvent.EventTags)
+			{
+				if (!bJustAdded)
+				{
+					GlobalEventsStats[Index].EventTags.FindOrAdd(CurrentEventTag.Key) += CurrentEventTag.Value;
+				}
+
+				if (EventToObjectsMap.Contains(CurrentEvent.Name + CurrentEventTag.Key.ToString()))
+				{
+					EventToObjectsMap[CurrentEvent.Name + CurrentEventTag.Key.ToString()].AddUnique(LogItem->GetName().ToString());
+				}
+				else
+				{
+					EventToObjectsMap.FindOrAdd(CurrentEvent.Name + CurrentEventTag.Key.ToString()).AddUnique(LogItem->GetName().ToString());
+				}
+			}
+		}
+	}
+
+	OutString.Append(TEXT("-= SUMMARY =-\n"));
+
+	for (auto& CurrentEvent : GlobalEventsStats)
+	{
+		OutString.Append(FString::Printf(TEXT("%s  occurred %d times by %d owners (%s)\n"), *CurrentEvent.Name, CurrentEvent.Counter, EventToObjectsMap.Contains(CurrentEvent.Name) ? EventToObjectsMap[CurrentEvent.Name].Num() : 0, *CurrentEvent.UserFriendlyDesc));
+		for (auto& CurrentEventTag : CurrentEvent.EventTags)
+		{
+			const int32 ObjectsNumber = EventToObjectsMap.Contains(CurrentEvent.Name + CurrentEventTag.Key.ToString()) ? EventToObjectsMap[CurrentEvent.Name + CurrentEventTag.Key.ToString()].Num() : 0;
+			OutString.Append(FString::Printf(TEXT("%s to <%s> tag occurred %d  times by %d owners (average %.2f times each)\n"), *CurrentEvent.Name, *CurrentEventTag.Key.ToString(), CurrentEventTag.Value, ObjectsNumber, ObjectsNumber > 0 ? float(CurrentEventTag.Value) / ObjectsNumber : -1));
+		}
+		OutString.Append(TEXT("\n"));
+	}
+
+	return FText::FromString(OutString);
+}
+
+
+#undef LOCTEXT_NAMESPACE
diff --git a/Engine/Source/Developer/NewLogVisualizer/Private/SVisualLoggerReport.h b/Engine/Source/Developer/NewLogVisualizer/Private/SVisualLoggerReport.h
new file mode 100644
--- /dev/null
+++ b/Engine/Source/Developer/NewLogVisualizer/Private/SVisualLoggerReport.h
@@ -0,0 +1,19 @@
+// Copyright 1998-2014 Epic Games, Inc. All Rights Reserved.
+
+#pragma once
+class SVisualLoggerReport : public SCompoundWidget
+{
+public:
+	SLATE_BEGIN_ARGS(SVisualLoggerReport) {}
+
+	SLATE_END_ARGS()
+
+		void Construct(const FArguments& InArgs, TArray< TSharedPtr<class STimeline> >&);
+
+	virtual ~SVisualLoggerReport();
+
+	FText GenerateReportText() const;
+
+protected:
+	TArray< TSharedPtr<class STimeline> > SelectedItems;
+};
diff --git a/Engine/Source/Developer/NewLogVisualizer/Private/VisualLoggerRenderingActor.cpp b/Engine/Source/Developer/NewLogVisualizer/Private/VisualLoggerRenderingActor.cpp
--- a/Engine/Source/Developer/NewLogVisualizer/Private/VisualLoggerRenderingActor.cpp
+++ b/Engine/Source/Developer/NewLogVisualizer/Private/VisualLoggerRenderingActor.cpp
@@ -155,6 +155,20 @@ void AVisualLoggerRenderingActor::AddDebugRendering()
 
 void AVisualLoggerRenderingActor::ObjectSelectionChanged(TSharedPtr<class STimeline> TimeLine)
 {
+	if (TimeLine.IsValid() == false)
+	{
+		Lines.Reset();
+		Points.Reset();
+		Cones.Reset();
+		Boxes.Reset();
+		Meshes.Reset();
+		Texts.Reset();
+		Cylinders.Reset();
+		Capsles.Reset();
+		MarkComponentsRenderStateDirty();
+		return;
+	}
+
 	const TArray<FVisualLogDevice::FVisualLogEntryItem>& Entries = TimeLine->GetEntries();
 	LogEntriesPath.Reset();
 	for (const auto &CurrentEntry : Entries)
diff --git a/Engine/Source/Editor/MainFrame/Private/Menus/MainMenu.h b/Engine/Source/Editor/MainFrame/Private/Menus/MainMenu.h
--- a/Engine/Source/Editor/MainFrame/Private/Menus/MainMenu.h
+++ b/Engine/Source/Editor/MainFrame/Private/Menus/MainMenu.h
@@ -99,6 +99,7 @@ private:
 	*/
 	static void OpenVisualLogger()
 	{
+		FModuleManager::Get().LoadModuleChecked<IModuleInterface>("NewLogVisualizer");
 		FGlobalTabmanager::Get()->InvokeTab(FName(TEXT("VisualLogger")));
 	}
 };
diff --git a/Engine/Source/Editor/UnrealEd/Private/Editor.cpp b/Engine/Source/Editor/UnrealEd/Private/Editor.cpp
--- a/Engine/Source/Editor/UnrealEd/Private/Editor.cpp
+++ b/Engine/Source/Editor/UnrealEd/Private/Editor.cpp
@@ -605,6 +605,11 @@ void UEditorEngine::Init(IEngineLoop* InEngineLoop)
 			FModuleManager::Get().LoadModule(TEXT("EnvironmentQueryEditor"));
 		}
 
+		if (GetDefault<UEditorExperimentalSettings>()->bVisualLogger)
+		{
+			FModuleManager::Get().LoadModule(TEXT("NewLogVisualizer"));
+		}
+
 		bool bGameplayAbilitiesEnabled = false;
 		GConfig->GetBool(TEXT("GameplayAbilities"), TEXT("GameplayAbilitiesEditorEnabled"), bGameplayAbilitiesEnabled, GEngineIni);
 		if (bGameplayAbilitiesEnabled)
diff --git a/Engine/Source/Editor/UnrealEd/Private/Settings/SettingsClasses.cpp b/Engine/Source/Editor/UnrealEd/Private/Settings/SettingsClasses.cpp
--- a/Engine/Source/Editor/UnrealEd/Private/Settings/SettingsClasses.cpp
+++ b/Engine/Source/Editor/UnrealEd/Private/Settings/SettingsClasses.cpp
@@ -56,10 +56,10 @@ UEditorExperimentalSettings::UEditorExperimentalSettings( const FObjectInitializ
 { 
 }
 
-
 void UEditorExperimentalSettings::PostEditChangeProperty( struct FPropertyChangedEvent& PropertyChangedEvent )
 {
 	static const FName NAME_EQS = GET_MEMBER_NAME_CHECKED(UEditorExperimentalSettings, bEQSEditor);
+	static const FName NAME_VISUAL_LOGGER = GET_MEMBER_NAME_CHECKED(UEditorExperimentalSettings, bVisualLogger);
 
 	Super::PostEditChangeProperty(PropertyChangedEvent);
 
@@ -76,6 +76,13 @@ void UEditorExperimentalSettings::PostEditChangeProperty( struct FPropertyChange
 			FModuleManager::Get().LoadModule(TEXT("EnvironmentQueryEditor"));
 		}
 	}
+	else if (Name == NAME_VISUAL_LOGGER)
+	{
+		if (bVisualLogger)
+		{
+			FModuleManager::Get().LoadModule(TEXT("NewLogVisualizer"));
+		}
+	}
 
 	if (!FUnrealEdMisc::Get().IsDeletePreferences())
 	{
@@ -102,7 +109,6 @@ void UEditorLoadingSavingSettings::SccHackInitialize()
 	bSCCUseGlobalSettings = ISourceControlModule::Get().GetUseGlobalSettings();
 }
 
-
 void UEditorLoadingSavingSettings::PostEditChangeProperty( struct FPropertyChangedEvent& PropertyChangedEvent )
 {
 	Super::PostEditChangeProperty(PropertyChangedEvent);
diff --git a/Engine/Source/Runtime/Engine/Private/VisualLogger/VisualLogger.cpp b/Engine/Source/Runtime/Engine/Private/VisualLogger/VisualLogger.cpp
--- a/Engine/Source/Runtime/Engine/Private/VisualLogger/VisualLogger.cpp
+++ b/Engine/Source/Runtime/Engine/Private/VisualLogger/VisualLogger.cpp
@@ -200,35 +200,38 @@ public:
 	{
 		if (FParse::Command(&Cmd, TEXT("VISLOG")))
 		{
-#if ENABLE_VISUAL_LOG
-			FString Command = FParse::Token(Cmd, 0);
-			if (Command == TEXT("record"))
-			{
-				FVisualLogger::Get().SetIsRecording(true);
-				return true;
-			}
-			else if (Command == TEXT("stop"))
-			{
-				FVisualLogger::Get().SetIsRecording(false);
-				return true;
-			}
-			else if (Command == TEXT("disableallbut"))
+			if (FModuleManager::Get().LoadModulePtr<IModuleInterface>("NewLogVisualizer") != nullptr)
 			{
-				FString Category = FParse::Token(Cmd, 1);
-				FVisualLogger::Get().BlockAllCategories(true);
-				FVisualLogger::Get().GetWhiteList().AddUnique(*Category);
-				return true;
-			}
+#if ENABLE_VISUAL_LOG
+				FString Command = FParse::Token(Cmd, 0);
+				if (Command == TEXT("record"))
+				{
+					FVisualLogger::Get().SetIsRecording(true);
+					return true;
+				}
+				else if (Command == TEXT("stop"))
+				{
+					FVisualLogger::Get().SetIsRecording(false);
+					return true;
+				}
+				else if (Command == TEXT("disableallbut"))
+				{
+					FString Category = FParse::Token(Cmd, 1);
+					FVisualLogger::Get().BlockAllCategories(true);
+					FVisualLogger::Get().GetWhiteList().AddUnique(*Category);
+					return true;
+				}
 #if WITH_EDITOR
-			else
-			{
-				FGlobalTabmanager::Get()->InvokeTab(FName(TEXT("VisualLogger")));
-				return true;
-			}
+				else
+				{
+					FGlobalTabmanager::Get()->InvokeTab(FName(TEXT("VisualLogger")));
+					return true;
+				}
 #endif
 #else
 			UE_LOG(LogVisual, Warning, TEXT("Unable to open LogVisualizer - logs are disabled"));
 #endif
+			}
 		}
 		return false;
 	}
