LogVisualizer search improvements
Optimisations to LogVisualizer data drawing
Fixed artifacts on LogVisualizer's timelines

diff --git a/Engine/Source/Developer/NewLogVisualizer/Private/LogVisualizerSettings.cpp b/Engine/Source/Developer/NewLogVisualizer/Private/LogVisualizerSettings.cpp
--- a/Engine/Source/Developer/NewLogVisualizer/Private/LogVisualizerSettings.cpp
+++ b/Engine/Source/Developer/NewLogVisualizer/Private/LogVisualizerSettings.cpp
@@ -13,6 +13,7 @@ ULogVisualizerSettings::ULogVisualizerSettings(const FObjectInitializer& ObjectI
 {
 	TrivialLogsThreshold = 1;
 	DefaultCameraDistance = 150;
+	bSearchInsideLogs = true;
 }
 
 void ULogVisualizerSettings::PostEditChangeProperty(struct FPropertyChangedEvent& PropertyChangedEvent)
diff --git a/Engine/Source/Developer/NewLogVisualizer/Private/STimeline.cpp b/Engine/Source/Developer/NewLogVisualizer/Private/STimeline.cpp
--- a/Engine/Source/Developer/NewLogVisualizer/Private/STimeline.cpp
+++ b/Engine/Source/Developer/NewLogVisualizer/Private/STimeline.cpp
@@ -52,27 +52,60 @@ void STimeline::UpdateVisibility()
 {
 	ULogVisualizerSettings* Settings = ULogVisualizerSettings::StaticClass()->GetDefaultObject<ULogVisualizerSettings>();
 	const bool bJustIgnore = Settings->bIgnoreTrivialLogs && Entries.Num() <= Settings->TrivialLogsThreshold;
-	SetVisibility(bJustIgnore || HiddenEntries.Num() == Entries.Num() || (SearchFilter.Len() > 0 && Name.ToString().Find(SearchFilter) == INDEX_NONE) ? EVisibility::Collapsed : EVisibility::Visible);
-	if (GetVisibility() != EVisibility::Visible)
+	const bool bIsCollapsed = bJustIgnore || HiddenEntries.Num() == Entries.Num() || (SearchFilter.Len() > 0 && Name.ToString().Find(SearchFilter) == INDEX_NONE);
+	SetVisibility(bIsCollapsed ? EVisibility::Collapsed : EVisibility::Visible);
+	if (bIsCollapsed)
 	{
 		Owner->SetSelectionState(SharedThis(this), false, false);
 	}
 }
 
-void STimeline::OnFiltersChanged()
+void STimeline::UpdateVisibilityForItems()
 {
 	HiddenEntries.Reset();
-
 	for (auto& CurrentEntry : Entries)
 	{
 		TArray<FVisualLoggerCategoryVerbosityPair> OutCategories;
 		FVisualLoggerHelpers::GetCategories(CurrentEntry.Entry, OutCategories);
-		if (VisualLoggerInterface->HasValidCategories(OutCategories) == false)
+		bool bHasValidCategories = VisualLoggerInterface->HasValidCategories(OutCategories);
+
+		const ULogVisualizerSettings* Settings = ULogVisualizerSettings::StaticClass()->GetDefaultObject<ULogVisualizerSettings>();
+		if (Settings->bSearchInsideLogs && bHasValidCategories && QuickSearchStrng.Len() > 0)
 		{
-			HiddenEntries.AddUnique(&CurrentEntry);
+			bool bMatchSearchString = false;
+			for (const FVisualLogLine& CurrentLine : CurrentEntry.Entry.LogLines)
+			{
+				if (CurrentLine.Line.Find(QuickSearchStrng) != INDEX_NONE || CurrentLine.Category.ToString().Find(QuickSearchStrng) != INDEX_NONE)
+				{
+					bMatchSearchString = true;
+					break;
+				}
+			}
+
+			if (!bMatchSearchString)
+			{
+				HiddenEntries.AddUnique(&CurrentEntry);
+			}
+		}
+		else
+		{
+			if (bHasValidCategories == false)
+			{
+				HiddenEntries.AddUnique(&CurrentEntry);
+			}
 		}
 	}
+}
 
+void STimeline::OnFiltersSearchChanged(const FText& Filter)
+{
+	QuickSearchStrng = Filter.ToString();
+	OnFiltersChanged();
+}
+
+void STimeline::OnFiltersChanged()
+{
+	UpdateVisibilityForItems();
 	UpdateVisibility();
 }
 
diff --git a/Engine/Source/Developer/NewLogVisualizer/Private/STimeline.h b/Engine/Source/Developer/NewLogVisualizer/Private/STimeline.h
--- a/Engine/Source/Developer/NewLogVisualizer/Private/STimeline.h
+++ b/Engine/Source/Developer/NewLogVisualizer/Private/STimeline.h
@@ -29,7 +29,10 @@ public:
 
 	void OnFiltersChanged();
 	void OnSearchChanged(const FText& Filter);
+	void OnFiltersSearchChanged(const FText& Filter);
 	void UpdateVisibility();
+	void UpdateVisibilityForItems();
+
 	void AddEntry(const FVisualLogDevice::FVisualLogEntryItem& Entry) { Entries.Add(Entry); }
 	const TArray<FVisualLogDevice::FVisualLogEntryItem>& GetEntries() { return Entries; }
 	FName GetName() { return Name; }
@@ -45,6 +48,7 @@ protected:
 
 	FName Name;
 	FString SearchFilter;
+	FString QuickSearchStrng;
 
 	/** Delegate to execute to get the menu content of this timeline */
 	FOnGetContent OnGetMenuContent;
diff --git a/Engine/Source/Developer/NewLogVisualizer/Private/STimelineBar.cpp b/Engine/Source/Developer/NewLogVisualizer/Private/STimelineBar.cpp
--- a/Engine/Source/Developer/NewLogVisualizer/Private/STimelineBar.cpp
+++ b/Engine/Source/Developer/NewLogVisualizer/Private/STimelineBar.cpp
@@ -310,82 +310,115 @@ int32 STimelineBar::OnPaint(const FPaintArgs& Args, const FGeometry& AllottedGeo
 	TArray<float> ErrorTimes;
 	TArray<float> WarningTimes;
 	auto &Entries = TimelineOwner.Pin()->GetEntries();
-	for (int32 Index = 0; Index < Entries.Num(); Index++)
+	int32 EntryIndex = 0;
+
+	while (EntryIndex < Entries.Num())
 	{
-		int32 CurrentIndex = Index;
-		float CurrentTime = Entries[CurrentIndex].Entry.TimeStamp;
-		if (CurrentTime < LocalViewRange.GetLowerBoundValue() || CurrentTime > LocalViewRange.GetUpperBoundValue())
+		const FVisualLogEntry& Entry = Entries[EntryIndex].Entry;
+		if (Entry.TimeStamp < LocalViewRange.GetLowerBoundValue() || Entry.TimeStamp > LocalViewRange.GetUpperBoundValue())
 		{
+			EntryIndex++;
 			continue;
 		}
 
-		if( TimelineOwner.Pin()->IsEntryHidden(Entries[CurrentIndex]) )
+		if (TimelineOwner.Pin()->IsEntryHidden(Entries[EntryIndex]))
 		{
+			EntryIndex++;
 			continue;
 		}
 
-		const TArray<FVisualLogLine>& LogLines = Entries[CurrentIndex].Entry.LogLines;
-		for (const FVisualLogLine& CurrentLine : LogLines)
+		// find bar width, connect all contiguous bars to draw them as one geometry (rendering optimization)
+		const float StartPos = (Entry.TimeStamp - LocalViewRange.GetLowerBoundValue()) * PixelsPerInput - 2;
+		float EndPos = (Entry.TimeStamp - LocalViewRange.GetLowerBoundValue()) * PixelsPerInput + 2;
+		int32 StartIndex = EntryIndex;
+		float LastEndX = MAX_FLT;
+		for (; StartIndex < Entries.Num(); ++StartIndex)
 		{
-			if (CurrentLine.Verbosity <= ELogVerbosity::Error)
+			const FVisualLogEntry& CurrentEntry = Entries[StartIndex].Entry;
+			if (CurrentEntry.TimeStamp < LocalViewRange.GetLowerBoundValue() || CurrentEntry.TimeStamp > LocalViewRange.GetUpperBoundValue())
 			{
-				ErrorTimes.AddUnique(CurrentTime);
 				break;
 			}
-			else if (CurrentLine.Verbosity == ELogVerbosity::Warning)
+
+			if (TimelineOwner.Pin()->IsEntryHidden(Entries[StartIndex]))
+			{
+				continue;
+			}
+
+			const TArray<FVisualLogLine>& LogLines = CurrentEntry.LogLines;
+			for (const FVisualLogLine& CurrentLine : LogLines)
+			{
+				if (CurrentLine.Verbosity <= ELogVerbosity::Error)
+				{
+					ErrorTimes.AddUnique(CurrentEntry.TimeStamp);
+					break;
+				}
+				else if (CurrentLine.Verbosity == ELogVerbosity::Warning)
+				{
+					WarningTimes.AddUnique(CurrentEntry.TimeStamp);
+					break;
+				}
+			}
+
+			const float CurrentStartPos = (CurrentEntry.TimeStamp - LocalViewRange.GetLowerBoundValue()) * PixelsPerInput - 2;
+			if (CurrentStartPos > EndPos)
 			{
-				WarningTimes.AddUnique(CurrentTime);
 				break;
 			}
+			EndPos = (CurrentEntry.TimeStamp - LocalViewRange.GetLowerBoundValue()) * PixelsPerInput + 2;
 		}
 
-		float LinePos = (CurrentTime - LocalViewRange.GetLowerBoundValue()) * PixelsPerInput;
-
-		FSlateDrawElement::MakeBox(
-			OutDrawElements,
-			RetLayerId++,
-			AllottedGeometry.ToPaintGeometry(
-			FVector2D(LinePos - BoxWidth * 0.5f, 0.0f),
-			FVector2D(BoxWidth, AllottedGeometry.Size.Y)),
-			FillImage,
-			MyClippingRect,
-			DrawEffects,
-			CurrentTimeColor
-			);
+		if (EndPos - StartPos > 0)
+		{
+			const float BarWidth = (EndPos - StartPos);
+			FSlateDrawElement::MakeBox(
+				OutDrawElements,
+				RetLayerId,
+				AllottedGeometry.ToPaintGeometry(
+				FVector2D(StartPos, 0.0f),
+				FVector2D(BarWidth, AllottedGeometry.Size.Y)),
+				FillImage,
+				MyClippingRect,
+				DrawEffects,
+				CurrentTimeColor
+				);
+		}
+		EntryIndex = StartIndex;
 	}
 
-	float ErrorBoxWidth = BoxWidth+2;
-	for (auto CurrentTime : ErrorTimes)
+	if (WarningTimes.Num()) RetLayerId++;
+	for (auto CurrentTime : WarningTimes)
 	{
 		float LinePos = (CurrentTime - LocalViewRange.GetLowerBoundValue()) * PixelsPerInput;
 
 		FSlateDrawElement::MakeBox(
 			OutDrawElements,
-			RetLayerId++,
+			RetLayerId,
 			AllottedGeometry.ToPaintGeometry(
-			FVector2D(LinePos - ErrorBoxWidth * 0.5f, 0.0f),
-			FVector2D(ErrorBoxWidth, AllottedGeometry.Size.Y)),
+			FVector2D(LinePos - 3, 0.0f),
+			FVector2D(6, AllottedGeometry.Size.Y)),
 			FillImage,
 			MyClippingRect,
 			DrawEffects,
-			ErrorTimeColor
+			WarningTimeColor
 			);
 	}
 
-	for (auto CurrentTime : WarningTimes)
+	if (ErrorTimes.Num()) RetLayerId++;
+	for (auto CurrentTime : ErrorTimes)
 	{
 		float LinePos = (CurrentTime - LocalViewRange.GetLowerBoundValue()) * PixelsPerInput;
 
 		FSlateDrawElement::MakeBox(
 			OutDrawElements,
-			RetLayerId++,
+			RetLayerId,
 			AllottedGeometry.ToPaintGeometry(
-			FVector2D(LinePos - ErrorBoxWidth * 0.5f, 0.0f),
-			FVector2D(ErrorBoxWidth, AllottedGeometry.Size.Y)),
+			FVector2D(LinePos - 3, 0.0f),
+			FVector2D(6, AllottedGeometry.Size.Y)),
 			FillImage,
 			MyClippingRect,
 			DrawEffects,
-			WarningTimeColor
+			ErrorTimeColor
 			);
 	}
 
@@ -402,13 +435,12 @@ int32 STimelineBar::OnPaint(const FPaintArgs& Args, const FGeometry& AllottedGeo
 		float CurrentTime = Entries[BestItemIndex].Entry.TimeStamp;
 		float LinePos = (CurrentTime - LocalViewRange.GetLowerBoundValue()) * PixelsPerInput;
 
-		BoxWidth += 2;
 		FSlateDrawElement::MakeBox(
 			OutDrawElements,
-			RetLayerId++,
+			++RetLayerId,
 			AllottedGeometry.ToPaintGeometry(
-			FVector2D(LinePos - BoxWidth * 0.5f, 0.0f),
-			FVector2D(BoxWidth, AllottedGeometry.Size.Y)),
+			FVector2D(LinePos - 2, 0.0f),
+			FVector2D(4, AllottedGeometry.Size.Y)),
 			SelectedFillImage,
 			MyClippingRect,
 			ESlateDrawEffect::None,
diff --git a/Engine/Source/Developer/NewLogVisualizer/Private/STimelinesContainer.cpp b/Engine/Source/Developer/NewLogVisualizer/Private/STimelinesContainer.cpp
--- a/Engine/Source/Developer/NewLogVisualizer/Private/STimelinesContainer.cpp
+++ b/Engine/Source/Developer/NewLogVisualizer/Private/STimelinesContainer.cpp
@@ -316,22 +316,16 @@ void STimelinesContainer::OnFiltersChanged()
 	}
 }
 
-void STimelinesContainer::GenerateReport()
+void STimelinesContainer::OnFiltersSearchChanged(const FText& Filter)
 {
-	//TArray< TSharedPtr<FLogsListItem> > ItemsToSave = LogsListWidget->GetSelectedItems();
-	//TArray< TSharedPtr<class STimeline> > AllLogs;
-
-	//TSharedPtr<FLogsListItem>* LogListItem = ItemsToSave.GetData();
-	//for (int32 ItemIndex = 0; ItemIndex < ItemsToSave.Num(); ++ItemIndex, ++LogListItem)
-	//{
-	//	if (LogListItem->IsValid() && LogVisualizer->Logs.IsValidIndex((*LogListItem)->LogIndex))
-	//	{
-	//		TSharedPtr<FActorsVisLog> Log = LogVisualizer->Logs[(*LogListItem)->LogIndex];
-	//		AllLogs.Add(Log);
-	//	}
-	//}
-
+	for (auto CurrentItem : TimelineItems)
+	{
+		CurrentItem->OnFiltersSearchChanged(Filter);
+	}
+}
 
+void STimelinesContainer::GenerateReport()
+{
 	TSharedRef<SWindow> NewWindow = SNew(SWindow)
 		.ClientSize(FVector2D(720, 768))
 		.Title(NSLOCTEXT("LogVisualizerReport", "WindowTitle", "Log Visualizer Report"))
diff --git a/Engine/Source/Developer/NewLogVisualizer/Private/STimelinesContainer.h b/Engine/Source/Developer/NewLogVisualizer/Private/STimelinesContainer.h
--- a/Engine/Source/Developer/NewLogVisualizer/Private/STimelinesContainer.h
+++ b/Engine/Source/Developer/NewLogVisualizer/Private/STimelinesContainer.h
@@ -30,6 +30,8 @@ public:
 	void OnNewLogEntry(const FVisualLogDevice::FVisualLogEntryItem& Entry);
 	void OnFiltersChanged();
 	void OnSearchChanged(const FText& Filter);
+	void OnFiltersSearchChanged(const FText& Filter);
+
 	void GenerateReport();
 
 	/**
diff --git a/Engine/Source/Developer/NewLogVisualizer/Private/SVisualLogger.cpp b/Engine/Source/Developer/NewLogVisualizer/Private/SVisualLogger.cpp
--- a/Engine/Source/Developer/NewLogVisualizer/Private/SVisualLogger.cpp
+++ b/Engine/Source/Developer/NewLogVisualizer/Private/SVisualLogger.cpp
@@ -675,7 +675,7 @@ void SVisualLogger::OnObjectSelectionChanged(TSharedPtr<class STimeline> TimeLin
 void SVisualLogger::OnFiltersSearchChanged(const FText& Filter)
 {
 	VisualLoggerFilters->OnFiltersSearchChanged(Filter);
-	LogsList->OnFiltersChanged();
-	MainView->OnFiltersChanged();
+	LogsList->OnFiltersSearchChanged(Filter);
+	MainView->OnFiltersSearchChanged(Filter);
 }
 #undef LOCTEXT_NAMESPACE
diff --git a/Engine/Source/Developer/NewLogVisualizer/Private/SVisualLoggerFilters.cpp b/Engine/Source/Developer/NewLogVisualizer/Private/SVisualLoggerFilters.cpp
--- a/Engine/Source/Developer/NewLogVisualizer/Private/SVisualLoggerFilters.cpp
+++ b/Engine/Source/Developer/NewLogVisualizer/Private/SVisualLoggerFilters.cpp
@@ -393,13 +393,15 @@ void SVisualLoggerFilters::AddFilter(const FString& GraphName, const FString& Da
 
 bool SVisualLoggerFilters::IsFilterEnabled(const FString& InFilterName, TEnumAsByte<ELogVerbosity::Type> Verbosity)
 {
+	const ULogVisualizerSettings* Settings = ULogVisualizerSettings::StaticClass()->GetDefaultObject<ULogVisualizerSettings>();
 	const FName FilterName(*InFilterName);
 	for (int32 Index = 0; Index < Filters.Num(); ++Index)
 	{
 		const SFilterWidget& Filter = Filters[Index].Get();
 		if (Filter.GetFilterName() == FilterName)
 		{
-			return Filter.IsEnabled() && Filter.GetVerbosity() >= Verbosity.GetValue() && (FiltersSearchString.Len() == 0 || Filter.GetFilterName().ToString().Find(FiltersSearchString) != INDEX_NONE);
+			return Filter.IsEnabled() && Filter.GetVerbosity() >= Verbosity.GetValue() && 
+				(Settings->bSearchInsideLogs == true || (Settings->bSearchInsideLogs == false && (FiltersSearchString.Len() == 0 || Filter.GetFilterName().ToString().Find(FiltersSearchString) != INDEX_NONE)));
 		}
 	}
 
diff --git a/Engine/Source/Developer/NewLogVisualizer/Private/SVisualLoggerLogsList.cpp b/Engine/Source/Developer/NewLogVisualizer/Private/SVisualLoggerLogsList.cpp
--- a/Engine/Source/Developer/NewLogVisualizer/Private/SVisualLoggerLogsList.cpp
+++ b/Engine/Source/Developer/NewLogVisualizer/Private/SVisualLoggerLogsList.cpp
@@ -7,6 +7,7 @@
 #include "SSequencerSectionOverlay.h"
 #include "STimelinesContainer.h"
 #include "SVisualLoggerLogsList.h"
+#include "LogVisualizerSettings.h"
 
 #define LOCTEXT_NAMESPACE "SVisualLoggerLogsList"
 
@@ -83,6 +84,7 @@ TSharedRef<ITableRow> SVisualLoggerLogsList::LogEntryLinesGenerateRow(TSharedPtr
 				SNew(STextBlock)
 				.ColorAndOpacity(FSlateColor(Item->CategoryColor))
 				.Text(Item->Category)
+				.HighlightText(this, &SVisualLoggerLogsList::GetFilterText)
 			]
 			+ SHorizontalBox::Slot()
 			.AutoWidth()
@@ -98,10 +100,24 @@ TSharedRef<ITableRow> SVisualLoggerLogsList::LogEntryLinesGenerateRow(TSharedPtr
 				SNew(STextBlock)
 				.AutoWrapText(true)
 				.Text(Item->Line)
+				.HighlightText(this, &SVisualLoggerLogsList::GetFilterText)
 			]
 		];
 }
 
+FText SVisualLoggerLogsList::GetFilterText() const
+{
+	static FText NoText;
+	const bool bSearchInsideLogs = ULogVisualizerSettings::StaticClass()->GetDefaultObject<ULogVisualizerSettings>()->bSearchInsideLogs;
+	return bSearchInsideLogs ? FilterText : NoText;
+}
+
+void SVisualLoggerLogsList::OnFiltersSearchChanged(const FText& Filter)
+{
+	FilterText = Filter;
+	OnFiltersChanged();
+}
+
 void SVisualLoggerLogsList::OnItemSelectionChanged(const FVisualLogDevice::FVisualLogEntryItem& LogEntry)
 {
 	CurrentLogEntry.Entry.Reset();
diff --git a/Engine/Source/Developer/NewLogVisualizer/Private/SVisualLoggerLogsList.h b/Engine/Source/Developer/NewLogVisualizer/Private/SVisualLoggerLogsList.h
--- a/Engine/Source/Developer/NewLogVisualizer/Private/SVisualLoggerLogsList.h
+++ b/Engine/Source/Developer/NewLogVisualizer/Private/SVisualLoggerLogsList.h
@@ -15,10 +15,13 @@ public:
 
 	void OnItemSelectionChanged(const FVisualLogDevice::FVisualLogEntryItem& EntryItem);
 	void OnFiltersChanged();
+	void OnFiltersSearchChanged(const FText& Filter);
+	FText GetFilterText() const;
 
 protected:
 	TSharedPtr<SListView<TSharedPtr<struct FLogEntryItem> > > LogsLinesWidget;
 	TArray<TSharedPtr<struct FLogEntryItem> > LogEntryLines;
 	TSharedPtr<IVisualLoggerInterface> VisualLoggerInterface;
 	FVisualLogDevice::FVisualLogEntryItem CurrentLogEntry;
+	FText FilterText;
 };
diff --git a/Engine/Source/Developer/NewLogVisualizer/Private/SVisualLoggerView.cpp b/Engine/Source/Developer/NewLogVisualizer/Private/SVisualLoggerView.cpp
--- a/Engine/Source/Developer/NewLogVisualizer/Private/SVisualLoggerView.cpp
+++ b/Engine/Source/Developer/NewLogVisualizer/Private/SVisualLoggerView.cpp
@@ -37,6 +37,8 @@ void SVisualLoggerView::Construct(const FArguments& InArgs, const TSharedRef<FUI
 		SNew(SScrollBar)
 		.Thickness(FVector2D(5.0f, 5.0f));
 
+	ULogVisualizerSettings* Settings = ULogVisualizerSettings::StaticClass()->GetDefaultObject<ULogVisualizerSettings>();
+
 	ChildSlot
 		[
 			SNew(SBorder)
@@ -93,7 +95,7 @@ void SVisualLoggerView::Construct(const FArguments& InArgs, const TSharedRef<FUI
 								[
 									SNew(SSearchBox)
 									.OnTextChanged(InArgs._OnFiltersSearchChanged)
-									.HintText(LOCTEXT("FiltersSearchHint", "Log Category Search"))
+									.HintText_Lambda([Settings]()->FText{return Settings->bSearchInsideLogs ? LOCTEXT("FiltersSearchHint", "Log Data Search") : LOCTEXT("FiltersSearchHint", "Log Category Search"); })
 								]
 #endif
 							]
@@ -210,4 +212,10 @@ void SVisualLoggerView::OnFiltersChanged()
 {
 	TimelinesContainer->OnFiltersChanged();
 }
+
+void SVisualLoggerView::OnFiltersSearchChanged(const FText& Filter)
+{
+	TimelinesContainer->OnFiltersSearchChanged(Filter);
+}
+
 #undef LOCTEXT_NAMESPACE
diff --git a/Engine/Source/Developer/NewLogVisualizer/Private/SVisualLoggerView.h b/Engine/Source/Developer/NewLogVisualizer/Private/SVisualLoggerView.h
--- a/Engine/Source/Developer/NewLogVisualizer/Private/SVisualLoggerView.h
+++ b/Engine/Source/Developer/NewLogVisualizer/Private/SVisualLoggerView.h
@@ -30,6 +30,7 @@ public:
 	void OnNewLogEntry(const FVisualLogDevice::FVisualLogEntryItem& Entry);
 	void OnFiltersChanged();
 	void OnSearchChanged(const FText& Filter);
+	void OnFiltersSearchChanged(const FText& Filter);
 	void OnSearchSplitterResized();
 	void OnObjectSelectionChanged(TSharedPtr<class STimeline> TimeLine);
 
diff --git a/Engine/Source/Developer/NewLogVisualizer/Public/LogVisualizerSettings.h b/Engine/Source/Developer/NewLogVisualizer/Public/LogVisualizerSettings.h
--- a/Engine/Source/Developer/NewLogVisualizer/Public/LogVisualizerSettings.h
+++ b/Engine/Source/Developer/NewLogVisualizer/Public/LogVisualizerSettings.h
@@ -44,6 +44,10 @@ public:
 	DECLARE_EVENT_OneParam(ULogVisualizerSettings, FSettingChangedEvent, FName /*PropertyName*/);
 	FSettingChangedEvent& OnSettingChanged() { return SettingChangedEvent; }
 
+	/**Whether to search/filter categories or to get text vlogs into account too */
+	UPROPERTY(EditAnywhere, config, Category = "VisualLogger")
+	bool bSearchInsideLogs;
+
 protected:
 
 	// UObject overrides
