Updated reports generated in LogVisualizer from collected Visual Logger events, mostly from visual point of view.
Added hyperlinks to reports generated in LogVisualizer - it works like filters for LogVisualzier and it should be easier to analyze report data now.

diff --git a/Engine/Source/Developer/NewLogVisualizer/Private/LogVisualizerStyle.cpp b/Engine/Source/Developer/NewLogVisualizer/Private/LogVisualizerStyle.cpp
--- a/Engine/Source/Developer/NewLogVisualizer/Private/LogVisualizerStyle.cpp
+++ b/Engine/Source/Developer/NewLogVisualizer/Private/LogVisualizerStyle.cpp
@@ -107,6 +107,91 @@ TSharedRef< FSlateStyleSet > FLogVisualizerStyle::Create()
 		Style.Set("Sequencer.SectionArea.Background", new FSlateColorBrush(FColor::White));
 	}
 
+	// Default text styles
+	{
+		const FTextBlockStyle NormalText = FTextBlockStyle()
+			.SetFont(TTF_FONT("Fonts/Roboto-Regular", 9))
+			.SetColorAndOpacity(FSlateColor::UseForeground())
+			.SetShadowOffset(FVector2D::ZeroVector)
+			.SetShadowColorAndOpacity(FLinearColor::Black)
+			.SetHighlightColor(FLinearColor(0.02f, 0.3f, 0.0f))
+			.SetHighlightShape(BOX_BRUSH("Common/TextBlockHighlightShape", FMargin(3.f / 8.f)));
+
+		Style.Set("RichText.Background", new BOX_BRUSH("Common/FlatColorSquare", FVector2D(1.0f, 1.0f), FMargin(0), FLinearColor(FColor(0xffeff3f3))));
+
+		Style.Set("RichText.RoundedBackground", new BOX_BRUSH("Common/RoundedSelection_16x", 4.0f / 16.0f, FLinearColor(FColor(0xffeff3f3))));
+
+		const FTextBlockStyle NormalRichTextStyle = FTextBlockStyle(NormalText)
+			.SetFont(TTF_FONT("Fonts/Roboto-Regular", 10))
+			.SetColorAndOpacity(FLinearColor(FColor(0xff2c3e50)));
+
+		Style.Set("RichText.Text", NormalRichTextStyle);
+		Style.Set("RichText.TextHighlight", FTextBlockStyle(NormalRichTextStyle)
+			.SetColorAndOpacity(FLinearColor(FColor(0xff19bc9c)))
+			);
+
+		Style.Set("RichText.TextBold", FTextBlockStyle(NormalRichTextStyle)
+			.SetFont(TTF_CORE_FONT("Fonts/Roboto-Bold", 10)));
+
+		Style.Set("RichText.HeaderText2", FTextBlockStyle(NormalRichTextStyle)
+			.SetFontSize(14));
+
+		Style.Set("RichText.HeaderText1", FTextBlockStyle(NormalRichTextStyle)
+			.SetFontSize(20));
+
+		{
+			const FButtonStyle RichTextHyperlinkButton = FButtonStyle()
+				.SetNormal(BORDER_BRUSH("Old/HyperlinkDotted", FMargin(0, 0, 0, 3 / 16.0f), FLinearColor::Blue))
+				.SetPressed(FSlateNoResource())
+				.SetHovered(BORDER_BRUSH("Old/HyperlinkUnderline", FMargin(0, 0, 0, 3 / 16.0f), FLinearColor::Blue));
+
+			const FTextBlockStyle RichTextHyperlinkText = FTextBlockStyle(NormalRichTextStyle)
+				.SetColorAndOpacity(FLinearColor::Blue);
+
+			Style.Set("RichText.HyperlinkText", RichTextHyperlinkText);
+
+			// legacy style
+			Style.Set("RichText.EditableHyperlinkText", RichTextHyperlinkText);
+
+			const FHyperlinkStyle RichTextHyperlink = FHyperlinkStyle()
+				.SetUnderlineStyle(RichTextHyperlinkButton)
+				.SetTextStyle(RichTextHyperlinkText)
+				.SetPadding(FMargin(0.0f));
+			Style.Set("RichText.Hyperlink", RichTextHyperlink);
+
+			Style.Set("RichText.ExternalLink", new IMAGE_BRUSH("Tutorials/ExternalLink", Icon16x16, FLinearColor::Blue));
+
+			// legacy style
+			Style.Set("RichText.EditableHyperlink", RichTextHyperlink);
+		}
+
+		//Tagline
+		{
+			Style.Set("RichText.Tagline.Background", new BOX_BRUSH("Common/FlatColorSquare", FVector2D(1.0f, 1.0f), FMargin(1), FLinearColor(FColor(0xffdbe4e4))));
+			Style.Set("RichText.Tagline.DarkBackground", new BOX_BRUSH("Common/FlatColorSquare", FVector2D(1.0f, 1.0f), FMargin(1), FLinearColor(0.55423, 0.60548, 0.60548)));
+			Style.Set("RichText.Tagline.Text", FTextBlockStyle(NormalText)
+				.SetFont(TTF_FONT("Fonts/Roboto-Bold", 24))
+				.SetColorAndOpacity(FLinearColor(FColor(0xff2c3e50)))
+				);
+
+			Style.Set("RichText.Tagline.TextHighlight", FTextBlockStyle(NormalText)
+				.SetFont(TTF_FONT("Fonts/Roboto-Bold", 24))
+				.SetColorAndOpacity(FLinearColor(FColor(0xff19bc9c)))
+				);
+
+			Style.Set("RichText.Tagline.SubtleText", FTextBlockStyle(NormalText)
+				.SetFont(TTF_FONT("Fonts/Roboto-Regular", 16))
+				.SetColorAndOpacity(FLinearColor(FColor(0xff2c3e50)))
+				);
+
+			Style.Set("RichText.Tagline.SubtleTextHighlight", FTextBlockStyle(NormalText)
+				.SetFont(TTF_FONT("Fonts/Roboto-Regular", 16))
+				.SetColorAndOpacity(FLinearColor(FColor(0xff19bc9c)))
+				);
+		}
+
+	}
+
 	return StyleRef;
 }
 
diff --git a/Engine/Source/Developer/NewLogVisualizer/Private/STimeline.cpp b/Engine/Source/Developer/NewLogVisualizer/Private/STimeline.cpp
--- a/Engine/Source/Developer/NewLogVisualizer/Private/STimeline.cpp
+++ b/Engine/Source/Developer/NewLogVisualizer/Private/STimeline.cpp
@@ -83,6 +83,17 @@ void STimeline::UpdateVisibilityForItems()
 					break;
 				}
 			}
+			if (!bMatchSearchString)
+			{
+				for (const FVisualLogEvent& CurrentEvent : CurrentEntry.Entry.Events)
+				{
+					if (CurrentEvent.Name.Find(QuickSearchStrng) != INDEX_NONE)
+					{
+						bMatchSearchString = true;
+						break;
+					}
+				}
+			}
 
 			if (!bMatchSearchString)
 			{
diff --git a/Engine/Source/Developer/NewLogVisualizer/Private/STimelinesContainer.cpp b/Engine/Source/Developer/NewLogVisualizer/Private/STimelinesContainer.cpp
--- a/Engine/Source/Developer/NewLogVisualizer/Private/STimelinesContainer.cpp
+++ b/Engine/Source/Developer/NewLogVisualizer/Private/STimelinesContainer.cpp
@@ -332,7 +332,7 @@ void STimelinesContainer::GenerateReport()
 		.ClientSize(FVector2D(720, 768))
 		.Title(NSLOCTEXT("LogVisualizerReport", "WindowTitle", "Log Visualizer Report"))
 		[
-			SNew(SVisualLoggerReport, SelectedNodes)
+			SNew(SVisualLoggerReport, SelectedNodes, VisualLoggerView)
 		];
 
 	FSlateApplication::Get().AddWindow(NewWindow);
diff --git a/Engine/Source/Developer/NewLogVisualizer/Private/SVisualLoggerReport.cpp b/Engine/Source/Developer/NewLogVisualizer/Private/SVisualLoggerReport.cpp
--- a/Engine/Source/Developer/NewLogVisualizer/Private/SVisualLoggerReport.cpp
+++ b/Engine/Source/Developer/NewLogVisualizer/Private/SVisualLoggerReport.cpp
@@ -14,40 +14,69 @@
 #include "Editor/UnrealEd/Public/EditorViewportClient.h"
 #include "VisualLoggerTypes.h"
 #include "SVisualLoggerReport.h"
+#include "SSearchBox.h"
 
 #define LOCTEXT_NAMESPACE "SVisualLoggerReport"
 
 BEGIN_SLATE_FUNCTION_BUILD_OPTIMIZATION
-void SVisualLoggerReport::Construct(const FArguments& InArgs, TArray< TSharedPtr<class STimeline> >& InSelectedItems)
+void SVisualLoggerReport::Construct(const FArguments& InArgs, TArray< TSharedPtr<class STimeline> >& InSelectedItems, TSharedPtr<class SVisualLoggerView> VisualLoggerView)
 {
 	SelectedItems = InSelectedItems;
+	GenerateReportText();
+
+	TArray< TSharedRef< ITextDecorator > > CustomDecorators;
+	for (auto& CurrentEvent : CollectedEvents)
+	{
+		CustomDecorators.Add(SRichTextBlock::HyperlinkDecorator(CurrentEvent, FSlateHyperlinkRun::FOnClick::CreateLambda(
+			[this, VisualLoggerView](const FSlateHyperlinkRun::FMetadata& Metadata){ VisualLoggerView->SetSearchString(FText::FromString(Metadata[TEXT("id")])); }
+		)));
+	}
+
 	this->ChildSlot
+	[
+		SNew(SBorder)
+		.BorderImage(FLogVisualizerStyle::Get().GetBrush("RichText.Background")).HAlign(EHorizontalAlignment::HAlign_Fill)
 		[
 			SNew(SScrollBox)
 			+ SScrollBox::Slot()
+			.HAlign(EHorizontalAlignment::HAlign_Fill)
 			[
-				SNew(SBorder)
-				.BorderImage(FEditorStyle::Get().GetBrush("ToolPanel.GroupBorder"))
+				SNew(SVerticalBox)
+				+ SVerticalBox::Slot().AutoHeight().HAlign(HAlign_Right).VAlign(VAlign_Center).Padding(15, 15)
 				[
-					SNew(SVerticalBox)
-					+ SVerticalBox::Slot().AutoHeight().Padding(0)
+					SNew(SHorizontalBox)
+					+ SHorizontalBox::Slot().MaxWidth(300)
 					[
-						SNew(SBorder)
-						.BorderImage(FEditorStyle::Get().GetBrush("Menu.Background"))
-						.Padding(0)
+						SNew(SSearchBox)
+						.OnTextChanged( FOnTextChanged::CreateLambda(
+						[this](FText NewText){
+							InteractiveRichText->SetHighlightText(NewText); 
+						}))
+					]
+				]
+
+				+ SVerticalBox::Slot()
+				[
+					SNew(SBorder)
+					.Padding(5.0f)
+					.BorderImage(FCoreStyle::Get().GetBrush("BoxShadow"))
+					[
+						SNew(SBorder).Padding(2).HAlign(HAlign_Left)
+						.BorderImage(FLogVisualizerStyle::Get().GetBrush("RichText.Background"))
 						[
-							SNew(SRichTextBlock)
-							.Text(this, &SVisualLoggerReport::GenerateReportText)
-							.TextStyle(FEditorStyle::Get(), "ContentBrowser.Filters.Text")
-							.DecoratorStyleSet(&FEditorStyle::Get())
-							.WrapTextAt(800)
+							SAssignNew(InteractiveRichText, SRichTextBlock)
+							.Text(ReportText)
+							.TextStyle(FLogVisualizerStyle::Get(), "RichText.Text")
+							.DecoratorStyleSet(&FLogVisualizerStyle::Get())
 							.Justification(ETextJustify::Left)
 							.Margin(FMargin(20))
+							.Decorators(CustomDecorators)
 						]
 					]
 				]
 			]
-		];
+		]
+	];
 }
 END_SLATE_FUNCTION_BUILD_OPTIMIZATION
 
@@ -55,13 +84,14 @@ SVisualLoggerReport::~SVisualLoggerReport()
 {
 }
 
-FText SVisualLoggerReport::GenerateReportText() const
+void SVisualLoggerReport::GenerateReportText()
 {
 	FString OutString;
 	TArray<FVisualLogEvent> GlobalEventsStats;
 	TMap<FString, TArray<FString> >	 EventToObjectsMap;
 
 
+	OutString.Append(TEXT("<RichText.HeaderText1>Report Details</>\n"));
 	for (TSharedPtr<class STimeline> LogItem : SelectedItems)
 	{
 		TArray<FVisualLogEvent> AllEvents;
@@ -94,17 +124,18 @@ FText SVisualLoggerReport::GenerateReportText() const
 		}
 
 		bool bPrintNextLine = false;
+
 		if (AllEvents.Num() > 0)
 		{
-			OutString.Append(FString::Printf(TEXT("%s\n"), *LogItem->GetName().ToString()));
+			OutString.Append(FString::Printf(TEXT("    <RichText.HeaderText2>%s</>"), *LogItem->GetName().ToString()));
 		}
 		for (auto& CurrentEvent : AllEvents)
 		{
 			for (auto& CurrentEventTag : CurrentEvent.EventTags)
 			{
-				OutString.Append(FString::Printf(TEXT("        %s  with <%s> tag occurred    %d times\n"), *CurrentEvent.Name, *CurrentEventTag.Key.ToString(), CurrentEventTag.Value));
+				OutString.Append(FString::Printf(TEXT(" \n        \u2022  <a id=\"%s\" style=\"RichText.Hyperlink\">%s</>  with <RichText.TextBold>%s</> tag occurred    <RichText.TextBold>%d times</>"), *CurrentEvent.Name, *CurrentEvent.Name, *CurrentEventTag.Key.ToString(), CurrentEventTag.Value));
 			}
-			OutString.Append(FString::Printf(TEXT("        %s occurred %d times\n\n"), *CurrentEvent.Name, CurrentEvent.Counter));
+			OutString.Append(FString::Printf(TEXT("\n        \u2022  <a id=\"%s\" style=\"RichText.Hyperlink\">%s</> occurred <RichText.TextBold>%d times</>"), *CurrentEvent.Name, *CurrentEvent.Name, CurrentEvent.Counter));
 
 			bool bJustAdded = false;
 			int32 Index = GlobalEventsStats.Find(FVisualLogEvent(CurrentEvent));
@@ -137,23 +168,27 @@ FText SVisualLoggerReport::GenerateReportText() const
 					EventToObjectsMap.FindOrAdd(CurrentEvent.Name + CurrentEventTag.Key.ToString()).AddUnique(LogItem->GetName().ToString());
 				}
 			}
+			OutString.Append(TEXT("\n"));
 		}
+		OutString.Append(TEXT("\n"));
 	}
 
-	OutString.Append(TEXT("-= SUMMARY =-\n"));
+	OutString.Append(TEXT("\n\n<RichText.HeaderText1>Report Summary</>\n"));
 
+	CollectedEvents.Reset();
 	for (auto& CurrentEvent : GlobalEventsStats)
 	{
-		OutString.Append(FString::Printf(TEXT("%s  occurred %d times by %d owners (%s)\n"), *CurrentEvent.Name, CurrentEvent.Counter, EventToObjectsMap.Contains(CurrentEvent.Name) ? EventToObjectsMap[CurrentEvent.Name].Num() : 0, *CurrentEvent.UserFriendlyDesc));
+		CollectedEvents.Add(*CurrentEvent.Name);
+		OutString.Append(FString::Printf(TEXT("    <a id=\"%s\" style=\"RichText.Hyperlink\">%s</>  occurred <RichText.TextBold>%d times</> by %d owners (%s)\n"), *CurrentEvent.Name, *CurrentEvent.Name, CurrentEvent.Counter, EventToObjectsMap.Contains(CurrentEvent.Name) ? EventToObjectsMap[CurrentEvent.Name].Num() : 0, *CurrentEvent.UserFriendlyDesc));
 		for (auto& CurrentEventTag : CurrentEvent.EventTags)
 		{
 			const int32 ObjectsNumber = EventToObjectsMap.Contains(CurrentEvent.Name + CurrentEventTag.Key.ToString()) ? EventToObjectsMap[CurrentEvent.Name + CurrentEventTag.Key.ToString()].Num() : 0;
-			OutString.Append(FString::Printf(TEXT("%s to <%s> tag occurred %d  times by %d owners (average %.2f times each)\n"), *CurrentEvent.Name, *CurrentEventTag.Key.ToString(), CurrentEventTag.Value, ObjectsNumber, ObjectsNumber > 0 ? float(CurrentEventTag.Value) / ObjectsNumber : -1));
+			OutString.Append(FString::Printf(TEXT("        \u2022  %s to <RichText.TextBold>%s</> tag occurred <RichText.TextBold>%d times</> by %d owners (average %.2f times each)\n"), *CurrentEvent.Name, *CurrentEventTag.Key.ToString(), CurrentEventTag.Value, ObjectsNumber, ObjectsNumber > 0 ? float(CurrentEventTag.Value) / ObjectsNumber : -1));
 		}
 		OutString.Append(TEXT("\n"));
 	}
 
-	return FText::FromString(OutString);
+	ReportText = FText::FromString(OutString);
 }
 
 
diff --git a/Engine/Source/Developer/NewLogVisualizer/Private/SVisualLoggerReport.h b/Engine/Source/Developer/NewLogVisualizer/Private/SVisualLoggerReport.h
--- a/Engine/Source/Developer/NewLogVisualizer/Private/SVisualLoggerReport.h
+++ b/Engine/Source/Developer/NewLogVisualizer/Private/SVisualLoggerReport.h
@@ -8,12 +8,17 @@ public:
 
 	SLATE_END_ARGS()
 
-		void Construct(const FArguments& InArgs, TArray< TSharedPtr<class STimeline> >&);
+	void Construct(const FArguments& InArgs, TArray< TSharedPtr<class STimeline> >&, TSharedPtr<class SVisualLoggerView> VisualLoggerView);
 
 	virtual ~SVisualLoggerReport();
 
-	FText GenerateReportText() const;
+protected:
+	void GenerateReportText();
 
 protected:
 	TArray< TSharedPtr<class STimeline> > SelectedItems;
+	TArray<TSharedRef<ITextDecorator>> Decorators;
+	TSharedPtr<class SRichTextBlock> InteractiveRichText;
+	FText ReportText;
+	TArray<FString> CollectedEvents;
 };
diff --git a/Engine/Source/Developer/NewLogVisualizer/Private/SVisualLoggerView.cpp b/Engine/Source/Developer/NewLogVisualizer/Private/SVisualLoggerView.cpp
--- a/Engine/Source/Developer/NewLogVisualizer/Private/SVisualLoggerView.cpp
+++ b/Engine/Source/Developer/NewLogVisualizer/Private/SVisualLoggerView.cpp
@@ -95,7 +95,7 @@ void SVisualLoggerView::Construct(const FArguments& InArgs, const TSharedRef<FUI
 								SNew(SBox)
 								.Padding(FMargin(0, 0, 4, 0))
 								[
-									SNew(SSearchBox)
+									SAssignNew(SearchBox, SSearchBox)
 									.OnTextChanged(InArgs._OnFiltersSearchChanged)
 									.HintText_Lambda([Settings]()->FText{return Settings->bSearchInsideLogs ? LOCTEXT("FiltersSearchHint", "Log Data Search") : LOCTEXT("FiltersSearchHint", "Log Category Search"); })
 								]
@@ -161,6 +161,14 @@ void SVisualLoggerView::Construct(const FArguments& InArgs, const TSharedRef<FUI
 
 }
 
+void SVisualLoggerView::SetSearchString(FText SearchString) 
+{ 
+	if (SearchBox.IsValid())
+	{
+		SearchBox->SetText(SearchString);
+	}
+}
+
 void SVisualLoggerView::OnObjectSelectionChanged(TSharedPtr<class STimeline> TimeLine)
 {
 	//FIXME: scroll to selected timeline (SebaK)
diff --git a/Engine/Source/Developer/NewLogVisualizer/Private/SVisualLoggerView.h b/Engine/Source/Developer/NewLogVisualizer/Private/SVisualLoggerView.h
--- a/Engine/Source/Developer/NewLogVisualizer/Private/SVisualLoggerView.h
+++ b/Engine/Source/Developer/NewLogVisualizer/Private/SVisualLoggerView.h
@@ -26,6 +26,7 @@ public:
 	void SetAnimationOutlinerFillPercentage(float FillPercentage) { AnimationOutlinerFillPercentage = FillPercentage; }
 
 	TSharedRef<SWidget> MakeSectionOverlay(TSharedRef<class FSequencerTimeSliderController> TimeSliderController, const TAttribute< TRange<float> >& ViewRange, const TAttribute<float>& ScrubPosition, bool bTopOverlay);
+	void SetSearchString(FText SearchString);
 
 	void OnNewLogEntry(const FVisualLogDevice::FVisualLogEntryItem& Entry);
 	void OnFiltersChanged();
@@ -40,6 +41,7 @@ protected:
 	TSharedPtr<class STimelinesContainer> TimelinesContainer;
 	TSharedPtr<class SSplitter> SearchSplitter;
 	TSharedPtr<class SScrollBox> ScrollBox;
+	TSharedPtr<class SSearchBox> SearchBox;
 
 	FVisualLoggerEvents	VisualLoggerEvents;
 	float AnimationOutlinerFillPercentage;
