Hopefully fixed non unity OSX compilation

diff --git a/Engine/Source/Developer/NewLogVisualizer/NewLogVisualizer.Build.cs b/Engine/Source/Developer/NewLogVisualizer/NewLogVisualizer.Build.cs
--- a/Engine/Source/Developer/NewLogVisualizer/NewLogVisualizer.Build.cs
+++ b/Engine/Source/Developer/NewLogVisualizer/NewLogVisualizer.Build.cs
@@ -18,6 +18,7 @@ public class NewLogVisualizer : ModuleRules
 			new string[] {
 				"Developer/TaskGraph/Public",
 				"Runtime/Engine/Classes",
+                "Editor/GeometryMode/Public"
 			}
 		);
 
@@ -32,6 +33,7 @@ public class NewLogVisualizer : ModuleRules
                 "EditorStyle",
 				"Engine",
 				"TaskGraph",
+                "GeometryMode",
 				"UnrealEd",
                 "SequencerWidgets",
 			}
diff --git a/Engine/Source/Developer/NewLogVisualizer/Private/LogVisualizerSettings.cpp b/Engine/Source/Developer/NewLogVisualizer/Private/LogVisualizerSettings.cpp
--- a/Engine/Source/Developer/NewLogVisualizer/Private/LogVisualizerSettings.cpp
+++ b/Engine/Source/Developer/NewLogVisualizer/Private/LogVisualizerSettings.cpp
@@ -6,6 +6,7 @@
 #if WITH_EDITOR
 #include "Editor/EditorEngine.h"
 #include "ISettingsModule.h"
+#include "GeometryEdMode.h"
 #include "UnrealEdMisc.h"
 #endif // WITH_EDITOR
 
@@ -16,6 +17,7 @@ ULogVisualizerSettings::ULogVisualizerSettings(const FObjectInitializer& ObjectI
 	bSearchInsideLogs = true;
 }
 
+#if WITH_EDITOR
 void ULogVisualizerSettings::PostEditChangeProperty(struct FPropertyChangedEvent& PropertyChangedEvent)
 {
 	Super::PostEditChangeProperty(PropertyChangedEvent);
@@ -28,4 +30,4 @@ void ULogVisualizerSettings::PostEditChangeProperty(struct FPropertyChangedEvent
 
 	SettingChangedEvent.Broadcast(Name);
 }
-
+#endif
diff --git a/Engine/Source/Developer/NewLogVisualizer/Public/LogVisualizerSettings.h b/Engine/Source/Developer/NewLogVisualizer/Public/LogVisualizerSettings.h
--- a/Engine/Source/Developer/NewLogVisualizer/Public/LogVisualizerSettings.h
+++ b/Engine/Source/Developer/NewLogVisualizer/Public/LogVisualizerSettings.h
@@ -12,14 +12,6 @@ UCLASS(config = EditorUserSettings)
 class NEWLOGVISUALIZER_API ULogVisualizerSettings : public UObject
 {
 	GENERATED_UCLASS_BODY()
-public:
-#if WITH_EDITOR
-	// Begin UObject Interface
-	//virtual void PostEditChangeProperty(struct FPropertyChangedEvent& PropertyChangedEvent) override;
-	//virtual void PostInitProperties() override;
-	// End UObject Interface
-#endif
-
 public:
 	/**Whether to show trivial logs, i.e. the ones with only one entry.*/
 	UPROPERTY(EditAnywhere, config, Category = "VisualLogger")
@@ -48,11 +40,10 @@ public:
 	UPROPERTY(EditAnywhere, config, Category = "VisualLogger")
 	bool bSearchInsideLogs;
 
-protected:
-
 	// UObject overrides
-
+#if WITH_EDITOR
 	virtual void PostEditChangeProperty(struct FPropertyChangedEvent& PropertyChangedEvent) override;
+#endif
 
 private:
 
