Fixed Cyclical Dependency in action graph issue

diff --git a/Engine/Source/Developer/NewLogVisualizer/NewLogVisualizer.Build.cs b/Engine/Source/Developer/NewLogVisualizer/NewLogVisualizer.Build.cs
--- a/Engine/Source/Developer/NewLogVisualizer/NewLogVisualizer.Build.cs
+++ b/Engine/Source/Developer/NewLogVisualizer/NewLogVisualizer.Build.cs
@@ -33,7 +33,6 @@ public class NewLogVisualizer : ModuleRules
                 "EditorStyle",
 				"Engine",
 				"TaskGraph",
-                "GeometryMode",
 				"UnrealEd",
                 "SequencerWidgets",
 			}
