LogVisualizer fixes and improvements (integrated from Dev to Main):
- Fixed and polished how we scroll/drag or zoom data in timelines using mouse. It should be much easier now.
- Added scroll bar while working with zoomed data. It's clear that we have it zoomed and it's faster to scroll that kind of data now.
- Fixed how timeline's context menu works with timelines selection.
- it's possible to copy log lines to system clipboard now,
- bunch of other fixes and tweaks here and there (like GrabHand" mouse cursor while dragging timelines, etc.)

diff --git a/Engine/Source/Developer/NewLogVisualizer/Private/SSequencerSectionOverlay.cpp b/Engine/Source/Developer/NewLogVisualizer/Private/SSequencerSectionOverlay.cpp
--- a/Engine/Source/Developer/NewLogVisualizer/Private/SSequencerSectionOverlay.cpp
+++ b/Engine/Source/Developer/NewLogVisualizer/Private/SSequencerSectionOverlay.cpp
@@ -13,7 +13,7 @@ int32 SSequencerSectionOverlay::OnPaint( const FPaintArgs& Args, const FGeometry
 
 FReply SSequencerSectionOverlay::OnMouseButtonDown( const FGeometry& MyGeometry, const FPointerEvent& MouseEvent )
 {
-	return TimeSliderController->OnMouseButtonDown( SharedThis(this), MyGeometry, MouseEvent );
+	return TimeSliderController->OnMouseButtonDown(SharedThis(this), MyGeometry, MouseEvent);
 }
 
 FReply SSequencerSectionOverlay::OnMouseButtonUp( const FGeometry& MyGeometry, const FPointerEvent& MouseEvent )
@@ -28,5 +28,9 @@ FReply SSequencerSectionOverlay::OnMouseMove( const FGeometry& MyGeometry, const
 
 FReply SSequencerSectionOverlay::OnMouseWheel( const FGeometry& MyGeometry, const FPointerEvent& MouseEvent )
 {
-	return TimeSliderController->OnMouseWheel( SharedThis(this), MyGeometry, MouseEvent );
+	if (MouseEvent.IsLeftShiftDown() || MouseEvent.IsLeftControlDown())
+	{
+		return TimeSliderController->OnMouseWheel(SharedThis(this), MyGeometry, MouseEvent);
+	}
+	return FReply::Unhandled();
 }
diff --git a/Engine/Source/Developer/NewLogVisualizer/Private/STimeline.cpp b/Engine/Source/Developer/NewLogVisualizer/Private/STimeline.cpp
--- a/Engine/Source/Developer/NewLogVisualizer/Private/STimeline.cpp
+++ b/Engine/Source/Developer/NewLogVisualizer/Private/STimeline.cpp
@@ -7,6 +7,32 @@
 #include "STimelinesContainer.h"
 #include "LogVisualizerSettings.h"
 
+class STimelineLabelAnchor : public SMenuAnchor
+{
+public:
+	void Construct(const FArguments& InArgs, TSharedPtr<class STimeline> InTimelineOwner)
+	{
+		SMenuAnchor::Construct(InArgs);
+		TimelineOwner = InTimelineOwner;
+	}
+
+private:
+	/** SWidget Interface */
+	virtual FReply OnMouseButtonDown(const FGeometry& MyGeometry, const FPointerEvent& MouseEvent) override;
+	TWeakPtr<class STimeline> TimelineOwner;
+};
+
+FReply STimelineLabelAnchor::OnMouseButtonDown(const FGeometry& MyGeometry, const FPointerEvent& MouseEvent)
+{
+	if (MouseEvent.GetEffectingButton() == EKeys::RightMouseButton && TimelineOwner.IsValid() && TimelineOwner.Pin()->IsSelected())
+	{
+		SetIsOpen(!IsOpen());
+		return FReply::Handled();
+	}
+	return FReply::Unhandled();
+}
+
+
 STimeline::~STimeline()
 {
 	ULogVisualizerSettings::StaticClass()->GetDefaultObject<ULogVisualizerSettings>()->OnSettingChanged().RemoveAll(this);
@@ -15,12 +41,8 @@ STimeline::~STimeline()
 FReply STimeline::OnMouseButtonDown(const FGeometry& MyGeometry, const FPointerEvent& MouseEvent)
 {
 	Owner->ChangeSelection(SharedThis(this), MouseEvent);
-	if (OnGetMenuContent.IsBound() && MouseEvent.GetEffectingButton() == EKeys::RightMouseButton)
-	{
-		PopupAnchor->SetIsOpen(!PopupAnchor->IsOpen());
-	}
 
-	return FReply::Handled();
+	return FReply::Unhandled();
 }
 
 FReply STimeline::OnMouseButtonUp(const FGeometry& MyGeometry, const FPointerEvent& MouseEvent)
@@ -161,7 +183,7 @@ void STimeline::Construct(const FArguments& InArgs, TSharedPtr<SVisualLoggerView
 			.VAlign(VAlign_Fill)
 			.FillWidth(TAttribute<float>(VisualLoggerView.Get(), &SVisualLoggerView::GetAnimationOutlinerFillPercentage))
 			[
-				SAssignNew(PopupAnchor, SMenuAnchor)
+				SAssignNew(PopupAnchor, STimelineLabelAnchor, SharedThis(this))
 				.OnGetMenuContent(OnGetMenuContent)
 				[
 					SNew(SBorder)
diff --git a/Engine/Source/Developer/NewLogVisualizer/Private/STimelineBar.cpp b/Engine/Source/Developer/NewLogVisualizer/Private/STimelineBar.cpp
--- a/Engine/Source/Developer/NewLogVisualizer/Private/STimelineBar.cpp
+++ b/Engine/Source/Developer/NewLogVisualizer/Private/STimelineBar.cpp
@@ -18,26 +18,20 @@
 FReply STimelineBar::OnMouseButtonDown(const FGeometry& MyGeometry, const FPointerEvent& MouseEvent)
 {
 	TimelineOwner.Pin()->OnMouseButtonDown(MyGeometry, MouseEvent);
-
-	if (MouseEvent.GetEffectingButton() == EKeys::LeftMouseButton)
-	{
-		return TimeSliderController->OnMouseButtonDown(SharedThis(this), MyGeometry, MouseEvent);
-	}
-	return FReply::Unhandled();
+	return TimeSliderController->OnMouseButtonDown(SharedThis(this), MyGeometry, MouseEvent);
 }
 
 FReply STimelineBar::OnMouseButtonUp(const FGeometry& MyGeometry, const FPointerEvent& MouseEvent)
 {
 	TimelineOwner.Pin()->OnMouseButtonUp(MyGeometry, MouseEvent);
 
+	FReply  Replay = TimeSliderController->OnMouseButtonUp(SharedThis(this), MyGeometry, MouseEvent);
 	if (MouseEvent.GetEffectingButton() == EKeys::LeftMouseButton)
 	{
-		FReply  Replay = TimeSliderController->OnMouseButtonUp(SharedThis(this), MyGeometry, MouseEvent);
 		SnapScrubPosition(TimeSliderController->GetTimeSliderArgs().ScrubPosition.Get());
-		return Replay;
 	}
 
-	return FReply::Unhandled();
+	return Replay;
 }
 
 FReply STimelineBar::OnMouseMove(const FGeometry& MyGeometry, const FPointerEvent& MouseEvent)
@@ -187,11 +181,11 @@ FReply STimelineBar::OnKeyDown(const FGeometry& MyGeometry, const FKeyEvent& InK
 			SnapScrubPosition(NewTimeStamp);
 			if (NewTimeStamp < LocalViewRange.GetLowerBoundValue())
 			{
-				TimeSliderController->GetTimeSliderArgs().ViewRange.Set(TRange<float>(NewTimeStamp, NewTimeStamp + RangeSize));
+				TimeSliderController->SetTimeRange(NewTimeStamp, NewTimeStamp + RangeSize);
 			}
 			else if (NewTimeStamp > LocalViewRange.GetUpperBoundValue())
 			{
-				TimeSliderController->GetTimeSliderArgs().ViewRange.Set(TRange<float>(NewTimeStamp - RangeSize, NewTimeStamp));
+				TimeSliderController->SetTimeRange(NewTimeStamp - RangeSize, NewTimeStamp);
 			}
 		}
 	}
diff --git a/Engine/Source/Developer/NewLogVisualizer/Private/SVisualLoggerLogsList.cpp b/Engine/Source/Developer/NewLogVisualizer/Private/SVisualLoggerLogsList.cpp
--- a/Engine/Source/Developer/NewLogVisualizer/Private/SVisualLoggerLogsList.cpp
+++ b/Engine/Source/Developer/NewLogVisualizer/Private/SVisualLoggerLogsList.cpp
@@ -42,12 +42,29 @@ void SVisualLoggerLogsList::Construct(const FArguments& InArgs, const TSharedRef
 			SAssignNew(LogsLinesWidget, SListView<TSharedPtr<FLogEntryItem> >)
 			.ItemHeight(20)
 			.ListItemsSource(&LogEntryLines)
-			.SelectionMode(ESelectionMode::Single)
+			.SelectionMode(ESelectionMode::Multi)
 			.OnSelectionChanged(this, &SVisualLoggerLogsList::LogEntryLineSelectionChanged)
 			.OnGenerateRow(this, &SVisualLoggerLogsList::LogEntryLinesGenerateRow)
 		];
 }
 
+FReply SVisualLoggerLogsList::OnKeyDown(const FGeometry& MyGeometry, const FKeyEvent& InKeyEvent)
+{
+	if (InKeyEvent.GetKey() == EKeys::C && (InKeyEvent.IsLeftCommandDown() || InKeyEvent.IsLeftControlDown()))
+	{
+		FString ClipboardString;
+		for (const TSharedPtr<struct FLogEntryItem>& CurrentItem : LogsLinesWidget->GetSelectedItems())
+		{
+			ClipboardString += CurrentItem->Category + FString(TEXT(" (")) + FString(FOutputDevice::VerbosityToString(CurrentItem->Verbosity)) + FString(TEXT(") ")) + CurrentItem->Line;
+			ClipboardString += TEXT("\n");
+		}
+		FPlatformMisc::ClipboardCopy(*ClipboardString);
+		return FReply::Handled();
+	}
+
+	return FReply::Unhandled();
+}
+
 void SVisualLoggerLogsList::OnFiltersChanged()
 {
 	FVisualLogDevice::FVisualLogEntryItem LogEntry = CurrentLogEntry;
diff --git a/Engine/Source/Developer/NewLogVisualizer/Private/SVisualLoggerLogsList.h b/Engine/Source/Developer/NewLogVisualizer/Private/SVisualLoggerLogsList.h
--- a/Engine/Source/Developer/NewLogVisualizer/Private/SVisualLoggerLogsList.h
+++ b/Engine/Source/Developer/NewLogVisualizer/Private/SVisualLoggerLogsList.h
@@ -12,6 +12,8 @@ public:
 
 	TSharedRef<ITableRow> LogEntryLinesGenerateRow(TSharedPtr<struct FLogEntryItem> Item, const TSharedRef<STableViewBase>& OwnerTable);
 	void LogEntryLineSelectionChanged(TSharedPtr<FLogEntryItem> SelectedItem, ESelectInfo::Type SelectInfo);
+	virtual bool SupportsKeyboardFocus() const override { return true; }
+	virtual FReply OnKeyDown(const FGeometry& MyGeometry, const FKeyEvent& InKeyEvent) override;
 
 	void OnItemSelectionChanged(const FVisualLogDevice::FVisualLogEntryItem& EntryItem);
 	void OnFiltersChanged();
diff --git a/Engine/Source/Developer/NewLogVisualizer/Private/SVisualLoggerView.cpp b/Engine/Source/Developer/NewLogVisualizer/Private/SVisualLoggerView.cpp
--- a/Engine/Source/Developer/NewLogVisualizer/Private/SVisualLoggerView.cpp
+++ b/Engine/Source/Developer/NewLogVisualizer/Private/SVisualLoggerView.cpp
@@ -11,13 +11,67 @@
 
 #define LOCTEXT_NAMESPACE "SVisualLoggerFilters"
 
+class SInputCatcherOverlay : public SOverlay
+{
+public:
+	void Construct(const FArguments& InArgs, TSharedRef<class FSequencerTimeSliderController> InTimeSliderController)
+	{
+		SOverlay::Construct(InArgs);
+		TimeSliderController = InTimeSliderController;
+	}
+
+	/** Controller for manipulating time */
+	TSharedPtr<class FSequencerTimeSliderController> TimeSliderController;
+private:
+	/** SWidget Interface */
+	virtual FReply OnMouseButtonDown(const FGeometry& MyGeometry, const FPointerEvent& MouseEvent) override;
+	virtual FReply OnMouseButtonUp(const FGeometry& MyGeometry, const FPointerEvent& MouseEvent) override;
+	virtual FReply OnMouseMove(const FGeometry& MyGeometry, const FPointerEvent& MouseEvent) override;
+	virtual FReply OnMouseWheel(const FGeometry& MyGeometry, const FPointerEvent& MouseEvent) override;
+
+private:
+};
+
+FReply SInputCatcherOverlay::OnMouseButtonDown(const FGeometry& MyGeometry, const FPointerEvent& MouseEvent)
+{
+	if (MouseEvent.GetEffectingButton() != EKeys::LeftMouseButton)
+	{
+		return TimeSliderController->OnMouseButtonDown(SharedThis(this), MyGeometry, MouseEvent);
+	}
+	return FReply::Unhandled();
+}
+
+FReply SInputCatcherOverlay::OnMouseButtonUp(const FGeometry& MyGeometry, const FPointerEvent& MouseEvent)
+{
+	if (MouseEvent.GetEffectingButton() != EKeys::LeftMouseButton)
+	{
+		return TimeSliderController->OnMouseButtonUp(SharedThis(this), MyGeometry, MouseEvent);
+	}
+	return FReply::Unhandled();
+}
+
+FReply SInputCatcherOverlay::OnMouseMove(const FGeometry& MyGeometry, const FPointerEvent& MouseEvent)
+{
+	return TimeSliderController->OnMouseMove(SharedThis(this), MyGeometry, MouseEvent);
+}
+
+FReply SInputCatcherOverlay::OnMouseWheel(const FGeometry& MyGeometry, const FPointerEvent& MouseEvent)
+{
+	if (MouseEvent.IsLeftShiftDown() || MouseEvent.IsLeftControlDown())
+	{
+		return TimeSliderController->OnMouseWheel(SharedThis(this), MyGeometry, MouseEvent);
+	}
+	return FReply::Unhandled();
+}
+
 void SVisualLoggerView::GetTimelines(TArray<TSharedPtr<STimeline> >& OutList, bool bOnlySelectedOnes)
 {
 	OutList = bOnlySelectedOnes ? TimelinesContainer->GetSelectedNodes() : TimelinesContainer->GetAllNodes();
 }
 
-void SVisualLoggerView::Construct(const FArguments& InArgs, const TSharedRef<FUICommandList>& InCommandList, TSharedPtr<IVisualLoggerInterface> VisualLoggerInterface)
+void SVisualLoggerView::Construct(const FArguments& InArgs, const TSharedRef<FUICommandList>& InCommandList, TSharedPtr<IVisualLoggerInterface> InVisualLoggerInterface)
 {
+	VisualLoggerInterface = InVisualLoggerInterface;
 	AnimationOutlinerFillPercentage = .25f;
 	VisualLoggerEvents = VisualLoggerInterface->GetVisualLoggerEvents();
 
@@ -27,7 +81,14 @@ void SVisualLoggerView::Construct(const FArguments& InArgs, const TSharedRef<FUI
 	TimeSliderArgs.ClampMax = InArgs._ViewRange.Get().GetUpperBoundValue();
 	TimeSliderArgs.ScrubPosition = InArgs._ScrubPosition;
 
+	TSharedRef<SScrollBar> ZoomScrollBar =
+		SNew(SScrollBar)
+		.Orientation(EOrientation::Orient_Horizontal)
+		.Thickness(FVector2D(2.0f, 2.0f));
+	ZoomScrollBar->SetState(0.0f, 1.0f);
+
 	TSharedPtr<FSequencerTimeSliderController> TimeSliderController(new FSequencerTimeSliderController(TimeSliderArgs));
+	TimeSliderController->SetExternalScrollbar(ZoomScrollBar);
 	VisualLoggerInterface->SetTimeSliderController(TimeSliderController);
 
 	// Create the top and bottom sliders
@@ -37,7 +98,8 @@ void SVisualLoggerView::Construct(const FArguments& InArgs, const TSharedRef<FUI
 
 	TSharedRef<SScrollBar> ScrollBar =
 		SNew(SScrollBar)
-		.Thickness(FVector2D(5.0f, 5.0f));
+		.Thickness(FVector2D(2.0f, 2.0f));
+
 
 	ULogVisualizerSettings* Settings = ULogVisualizerSettings::StaticClass()->GetDefaultObject<ULogVisualizerSettings>();
 
@@ -106,7 +168,7 @@ void SVisualLoggerView::Construct(const FArguments& InArgs, const TSharedRef<FUI
 					+ SVerticalBox::Slot()
 					.FillHeight(1.0)
 					[
-						SNew(SOverlay)
+						SNew(SInputCatcherOverlay, TimeSliderController.ToSharedRef())
 						+ SOverlay::Slot()
 						[
 							MakeSectionOverlay(TimeSliderController.ToSharedRef(), InArgs._ViewRange, InArgs._ScrubPosition, false)
@@ -115,16 +177,38 @@ void SVisualLoggerView::Construct(const FArguments& InArgs, const TSharedRef<FUI
 						[
 							SAssignNew(ScrollBox, SScrollBox)
 							.ExternalScrollbar(ScrollBar)
+							+ SScrollBox::Slot()
+							[
+								SAssignNew(TimelinesContainer, STimelinesContainer, SharedThis(this), TimeSliderController.ToSharedRef())
+								.VisualLoggerInterface(VisualLoggerInterface)
+							]
 						]
 						+ SOverlay::Slot()
 						[
 							MakeSectionOverlay(TimeSliderController.ToSharedRef(), InArgs._ViewRange, InArgs._ScrubPosition, true)
 						]
-						+SOverlay::Slot()
+						+ SOverlay::Slot()
 						.HAlign(HAlign_Right)
 						[
 							ScrollBar
 						]
+						+ SOverlay::Slot()
+						.VAlign(VAlign_Bottom)
+						[
+							SNew(SHorizontalBox)
+							+ SHorizontalBox::Slot()
+							//.Visibility(EVisibility::HitTestInvisible)
+							.FillWidth(TAttribute<float>(this, &SVisualLoggerView::GetAnimationOutlinerFillPercentage))
+							[
+								// Take up space but display nothing. This is required so that all areas dependent on time align correctly
+								SNullWidget::NullWidget
+							]
+							+ SHorizontalBox::Slot()
+							.FillWidth(1.0f)
+							[
+								ZoomScrollBar
+							]
+						]
 					]
 					+ SVerticalBox::Slot()
 					.AutoHeight()
@@ -153,14 +237,19 @@ void SVisualLoggerView::Construct(const FArguments& InArgs, const TSharedRef<FUI
 			]
 		];
 
-		ScrollBox->AddSlot()
-		[
-			SAssignNew(TimelinesContainer, STimelinesContainer, SharedThis(this), TimeSliderController.ToSharedRef())
-			.VisualLoggerInterface(VisualLoggerInterface)
-		];
+		//ScrollBox->AddSlot()
+		//[
+		//	SAssignNew(TimelinesContainer, STimelinesContainer, SharedThis(this), TimeSliderController.ToSharedRef())
+		//	.VisualLoggerInterface(VisualLoggerInterface)
+		//];
 
 }
 
+void SVisualLoggerView::SetAnimationOutlinerFillPercentage(float FillPercentage) 
+{ 
+	AnimationOutlinerFillPercentage = FillPercentage; 
+}
+
 void SVisualLoggerView::SetSearchString(FText SearchString) 
 { 
 	if (SearchBox.IsValid())
@@ -204,7 +293,7 @@ TSharedRef<SWidget> SVisualLoggerView::MakeSectionOverlay(TSharedRef<FSequencerT
 			// Take up space but display nothing. This is required so that all areas dependent on time align correctly
 			SNullWidget::NullWidget
 		]
-	+ SHorizontalBox::Slot()
+		+ SHorizontalBox::Slot()
 		.FillWidth(1.0f)
 		[
 			SNew(SSequencerSectionOverlay, TimeSliderController)
@@ -228,4 +317,14 @@ void SVisualLoggerView::OnFiltersSearchChanged(const FText& Filter)
 	TimelinesContainer->OnFiltersSearchChanged(Filter);
 }
 
+FCursorReply SVisualLoggerView::OnCursorQuery(const FGeometry& MyGeometry, const FPointerEvent& CursorEvent) const
+{
+	if (VisualLoggerInterface->GetTimeSliderController()->IsPanning())
+	{
+		return FCursorReply::Cursor(EMouseCursor::GrabHand);
+	}
+	return FCursorReply::Cursor(EMouseCursor::Default);
+}
+
+
 #undef LOCTEXT_NAMESPACE
diff --git a/Engine/Source/Developer/NewLogVisualizer/Private/SVisualLoggerView.h b/Engine/Source/Developer/NewLogVisualizer/Private/SVisualLoggerView.h
--- a/Engine/Source/Developer/NewLogVisualizer/Private/SVisualLoggerView.h
+++ b/Engine/Source/Developer/NewLogVisualizer/Private/SVisualLoggerView.h
@@ -23,7 +23,7 @@ public:
 
 		return LeftSplitterSlot.SizeValue.Get()/ RightSplitterSlot.SizeValue.Get();
 	}
-	void SetAnimationOutlinerFillPercentage(float FillPercentage) { AnimationOutlinerFillPercentage = FillPercentage; }
+	void SetAnimationOutlinerFillPercentage(float FillPercentage);
 
 	TSharedRef<SWidget> MakeSectionOverlay(TSharedRef<class FSequencerTimeSliderController> TimeSliderController, const TAttribute< TRange<float> >& ViewRange, const TAttribute<float>& ScrubPosition, bool bTopOverlay);
 	void SetSearchString(FText SearchString);
@@ -36,8 +36,10 @@ public:
 	void OnObjectSelectionChanged(TSharedPtr<class STimeline> TimeLine);
 
 	void GetTimelines(TArray<TSharedPtr<class STimeline> >&, bool bOnlySelectedOnes = false);
+	virtual FCursorReply OnCursorQuery(const FGeometry& MyGeometry, const FPointerEvent& CursorEvent) const override;
 
 protected:
+	TSharedPtr<struct IVisualLoggerInterface> VisualLoggerInterface;
 	TSharedPtr<class STimelinesContainer> TimelinesContainer;
 	TSharedPtr<class SSplitter> SearchSplitter;
 	TSharedPtr<class SScrollBox> ScrollBox;
diff --git a/Engine/Source/Developer/NewLogVisualizer/Private/TimeSliderController.cpp b/Engine/Source/Developer/NewLogVisualizer/Private/TimeSliderController.cpp
--- a/Engine/Source/Developer/NewLogVisualizer/Private/TimeSliderController.cpp
+++ b/Engine/Source/Developer/NewLogVisualizer/Private/TimeSliderController.cpp
@@ -2,7 +2,6 @@
 
 #include "LogVisualizer.h"
 #include "TimeSliderController.h"
-
 #define LOCTEXT_NAMESPACE "TimeSlider"
 
 namespace ScrubConstants
@@ -345,7 +344,8 @@ FReply FSequencerTimeSliderController::OnMouseButtonUp( TSharedRef<SWidget> Widg
 		}
 		
 		bPanning = false;
-		
+		FReply::Handled().CaptureMouse(WidgetOwner).UseHighPrecisionMouseMovement(WidgetOwner);
+
 		return FReply::Handled().ReleaseMouseCapture();
 	}
 	else if ( bHandleLeftMouseButton )
@@ -375,18 +375,22 @@ FReply FSequencerTimeSliderController::OnMouseMove( TSharedRef<SWidget> WidgetOw
 {
 	if ( WidgetOwner->HasMouseCapture() )
 	{
-		if (MouseEvent.IsMouseButtonDown( EKeys::RightMouseButton ))
+		if (MouseEvent.IsMouseButtonDown(EKeys::RightMouseButton))
 		{
 			if (!bPanning)
 			{
 				DistanceDragged += FMath::Abs( MouseEvent.GetCursorDelta().X );
 				if ( DistanceDragged > FSlateApplication::Get().GetDragTriggerDistnace() )
 				{
+					FReply::Handled().CaptureMouse(WidgetOwner).UseHighPrecisionMouseMovement(WidgetOwner);
+					SoftwareCursorPosition = MyGeometry.AbsoluteToLocal(MouseEvent.GetLastScreenSpacePosition());
 					bPanning = true;
 				}
 			}
 			else
 			{
+				SoftwareCursorPosition = MyGeometry.AbsoluteToLocal(MouseEvent.GetLastScreenSpacePosition());
+
 				TRange<float> LocalViewRange = TimeSliderArgs.ViewRange.Get();
 				float LocalViewRangeMin = LocalViewRange.GetLowerBoundValue();
 				float LocalViewRangeMax = LocalViewRange.GetUpperBoundValue();
@@ -414,6 +418,12 @@ FReply FSequencerTimeSliderController::OnMouseMove( TSharedRef<SWidget> WidgetOw
 				}
 
 				TimeSliderArgs.OnViewRangeChanged.ExecuteIfBound(TRange<float>(NewViewOutputMin, NewViewOutputMax));
+				if (Scrollbar.IsValid())
+				{
+					float InOffsetFraction = NewViewOutputMin / (LocalClampMax.GetValue() - LocalClampMin.GetValue());
+					float InThumbSizeFraction = (NewViewOutputMax - NewViewOutputMin) / (LocalClampMax.GetValue() - LocalClampMin.GetValue());
+					Scrollbar->SetState(InOffsetFraction, InThumbSizeFraction);
+				}
 
 				if( !TimeSliderArgs.ViewRange.IsBound() )
 				{	
@@ -459,6 +469,45 @@ void FSequencerTimeSliderController::CommitScrubPosition( float NewValue, bool b
 	TimeSliderArgs.OnScrubPositionChanged.ExecuteIfBound( NewValue, bIsScrubbing );
 }
 
+void FSequencerTimeSliderController::SetExternalScrollbar(TSharedRef<SScrollBar> InScrollbar) 
+{ 
+	Scrollbar = InScrollbar;
+	Scrollbar->SetOnUserScrolled(FOnUserScrolled::CreateRaw(this, &FSequencerTimeSliderController::HorizontalScrollBar_OnUserScrolled)); 
+};
+
+void FSequencerTimeSliderController::HorizontalScrollBar_OnUserScrolled(float ScrollOffset)
+{
+	if (!TimeSliderArgs.ViewRange.IsBound())
+	{
+		TRange<float> LocalViewRange = TimeSliderArgs.ViewRange.Get();
+		float LocalViewRangeMin = LocalViewRange.GetLowerBoundValue();
+		float LocalViewRangeMax = LocalViewRange.GetUpperBoundValue();
+		TOptional<float> LocalClampMin = TimeSliderArgs.ClampMin.Get();
+		TOptional<float> LocalClampMax = TimeSliderArgs.ClampMax.Get();
+
+		float NewViewOutputMin = ScrollOffset * (LocalClampMax.GetValue() - LocalClampMin.GetValue());
+		// The  output is not bound to a delegate so we'll manage the value ourselves
+		float NewViewOutputMax = NewViewOutputMin + (LocalViewRangeMax - LocalViewRangeMin);
+		TimeSliderArgs.ViewRange.Set(TRange<float>(NewViewOutputMin, NewViewOutputMax));
+
+		float InOffsetFraction = NewViewOutputMin / (LocalClampMax.GetValue() - LocalClampMin.GetValue());
+		float InThumbSizeFraction = (NewViewOutputMax - NewViewOutputMin) / (LocalClampMax.GetValue() - LocalClampMin.GetValue());
+		Scrollbar->SetState(InOffsetFraction, InThumbSizeFraction);
+	}
+}
+
+void FSequencerTimeSliderController::SetTimeRange(float NewViewOutputMin, float NewViewOutputMax)
+{
+	TimeSliderArgs.ViewRange.Set(TRange<float>(NewViewOutputMin, NewViewOutputMax));
+
+	TOptional<float> LocalClampMin = TimeSliderArgs.ClampMin.Get();
+	TOptional<float> LocalClampMax = TimeSliderArgs.ClampMax.Get();
+
+	float InOffsetFraction = NewViewOutputMin / (LocalClampMax.GetValue() - LocalClampMin.GetValue());
+	float InThumbSizeFraction = (NewViewOutputMax - NewViewOutputMin) / (LocalClampMax.GetValue() - LocalClampMin.GetValue());
+	Scrollbar->SetState(InOffsetFraction, InThumbSizeFraction);
+}
+
 FReply FSequencerTimeSliderController::OnMouseWheel( TSharedRef<SWidget> WidgetOwner, const FGeometry& MyGeometry, const FPointerEvent& MouseEvent )
 {
 	FReply ReturnValue = FReply::Unhandled();;
@@ -501,7 +550,20 @@ FReply FSequencerTimeSliderController::OnMouseWheel( TSharedRef<SWidget> WidgetO
 				}
 
 				TimeSliderArgs.OnViewRangeChanged.ExecuteIfBound(TRange<float>(NewViewOutputMin, NewViewOutputMax));
-
+				if (Scrollbar.IsValid())
+				{
+					/**
+					* Set the offset and size of the track's thumb.
+					* Note that the maximum offset is 1.0-ThumbSizeFraction.
+					* If the user can view 1/3 of the items in a single page, the maximum offset will be ~0.667f
+					*
+					* @param InOffsetFraction     Offset of the thumbnail from the top as a fraction of the total available scroll space.
+					* @param InThumbSizeFraction  Size of thumbnail as a fraction of the total available scroll space.
+					*/
+					float InOffsetFraction = NewViewOutputMin / (LocalClampMax.GetValue() - LocalClampMin.GetValue());
+					float InThumbSizeFraction = (NewViewOutputMax - NewViewOutputMin) / (LocalClampMax.GetValue() - LocalClampMin.GetValue());
+					Scrollbar->SetState(InOffsetFraction, InThumbSizeFraction);
+				}
 				if( !TimeSliderArgs.ViewRange.IsBound() )
 				{	
 					// The  output is not bound to a delegate so we'll manage the value ourselves
@@ -557,7 +619,7 @@ int32 FSequencerTimeSliderController::OnPaintSectionView( const FGeometry& Allot
 
 		FSlateDrawElement::MakeBox(
 			OutDrawElements,
-			LayerId + 1,
+			++LayerId,
 			CursorGeometry,
 			CursorBackground,
 			MyClippingRect,
@@ -573,7 +635,7 @@ int32 FSequencerTimeSliderController::OnPaintSectionView( const FGeometry& Allot
 
 		FSlateDrawElement::MakeLines(
 			OutDrawElements,
-			LayerId+2,
+			++LayerId,
 			AllottedGeometry.ToPaintGeometry( FVector2D(LinePos, 0.0f ), FVector2D(1.0f,1.0f) ),
 			LinePoints,
 			MyClippingRect,
diff --git a/Engine/Source/Developer/NewLogVisualizer/Private/TimeSliderController.h b/Engine/Source/Developer/NewLogVisualizer/Private/TimeSliderController.h
--- a/Engine/Source/Developer/NewLogVisualizer/Private/TimeSliderController.h
+++ b/Engine/Source/Developer/NewLogVisualizer/Private/TimeSliderController.h
@@ -43,6 +43,10 @@ public:
 	 */
 	void CommitScrubPosition( float NewValue, bool bIsScrubbing );
 
+	void SetExternalScrollbar(TSharedRef<SScrollBar> Scrollbar);
+	void SetTimeRange(float MinValue, float MaxValue);
+	bool IsPanning() { return bPanning; }
+
 private:
 
 	/**
@@ -53,6 +57,7 @@ private:
 	 * @param InArgs			Parameters for drawing the tick lines
 	 */
 	void DrawTicks( FSlateWindowElementList& OutDrawElements, const struct FScrubRangeToScreen& RangeToScreen, struct FDrawTickArgs& InArgs ) const;
+	void HorizontalScrollBar_OnUserScrolled(float ScrollOffset);
 private:
 	FVisualLoggerTimeSliderArgs TimeSliderArgs;
 	/** Brush for drawing an upwards facing scrub handle */
@@ -67,4 +72,7 @@ private:
 	bool bDraggingScrubber;
 	/** If we are currently panning the panel */
 	bool bPanning;
+	/***/
+	TSharedPtr<SScrollBar> Scrollbar;
+	FVector2D SoftwareCursorPosition;
 };
