Fixed movement on timelines with keyboard cursors in LogVisualizer tool (integrated from Dev to Main)

diff --git a/Engine/Source/Developer/NewLogVisualizer/Private/STimelineBar.cpp b/Engine/Source/Developer/NewLogVisualizer/Private/STimelineBar.cpp
--- a/Engine/Source/Developer/NewLogVisualizer/Private/STimelineBar.cpp
+++ b/Engine/Source/Developer/NewLogVisualizer/Private/STimelineBar.cpp
@@ -178,7 +178,7 @@ FReply STimelineBar::OnKeyDown(const FGeometry& MyGeometry, const FKeyEvent& InK
 		if (NewItemIndex != CurrentItemIndex)
 		{
 			float NewTimeStamp = Entries[NewItemIndex].Entry.TimeStamp;
-			SnapScrubPosition(NewTimeStamp);
+			SnapScrubPosition(NewItemIndex);
 			if (NewTimeStamp < LocalViewRange.GetLowerBoundValue())
 			{
 				TimeSliderController->SetTimeRange(NewTimeStamp, NewTimeStamp + RangeSize);
@@ -239,6 +239,21 @@ void STimelineBar::SnapScrubPosition(float ScrubPosition)
 	}
 }
 
+void STimelineBar::SnapScrubPosition(int32 NewItemIndex)
+{
+	if (NewItemIndex != INDEX_NONE)
+	{
+		auto &Entries = TimelineOwner.Pin()->GetEntries();
+		const float CurrentTime = Entries[NewItemIndex].Entry.TimeStamp;
+		if (CurrentItemIndex != NewItemIndex)
+		{
+			CurrentItemIndex = NewItemIndex;
+			VisualLoggerInterface->GetVisualLoggerEvents().OnItemSelectionChanged.ExecuteIfBound(Entries[NewItemIndex]);
+		}
+		TimeSliderController->CommitScrubPosition(CurrentTime, false);
+	}
+}
+
 void STimelineBar::Construct(const FArguments& InArgs, TSharedPtr<FSequencerTimeSliderController> InTimeSliderController, TSharedPtr<STimeline> InTimelineOwner)
 {
 	VisualLoggerInterface = InArgs._VisualLoggerInterface.Get();
diff --git a/Engine/Source/Developer/NewLogVisualizer/Private/STimelineBar.h b/Engine/Source/Developer/NewLogVisualizer/Private/STimelineBar.h
--- a/Engine/Source/Developer/NewLogVisualizer/Private/STimelineBar.h
+++ b/Engine/Source/Developer/NewLogVisualizer/Private/STimelineBar.h
@@ -27,6 +27,7 @@ public:
 	virtual int32 OnPaint(const FPaintArgs& Args, const FGeometry& AllottedGeometry, const FSlateRect& MyClippingRect, FSlateWindowElementList& OutDrawElements, int32 LayerId, const FWidgetStyle& InWidgetStyle, bool bParentEnabled) const override;
 	FVector2D ComputeDesiredSize() const;
 	void SnapScrubPosition(float ScrubPosition);
+	void SnapScrubPosition(int32 NewItemIndex);
 	uint32 GetClosestItem(float Time) const;
 
 	void OnSelect();
