LogVisualization improvements integrated from Dev to Main:
- Added toggle button to LogVisualizer's toolbar, to toggle graphs 2d visualization
- User can change background color for 2d graphs in LogVisualizer's settings

diff --git a/Engine/Source/Developer/NewLogVisualizer/Private/LogVisualizer.h b/Engine/Source/Developer/NewLogVisualizer/Private/LogVisualizer.h
--- a/Engine/Source/Developer/NewLogVisualizer/Private/LogVisualizer.h
+++ b/Engine/Source/Developer/NewLogVisualizer/Private/LogVisualizer.h
@@ -60,4 +60,5 @@ protected:
 #include "SVisualLoggerLogsList.h"
 #include "SVisualLoggerStatusView.h"
 #include "STimeline.h"
-#include "LogVisualizerSettings.h"
\ No newline at end of file
+#include "LogVisualizerSettings.h"
+#include "LogVisualizerSessionSettings.h"
\ No newline at end of file
diff --git a/Engine/Source/Developer/NewLogVisualizer/Private/LogVisualizerSessionSettings.cpp b/Engine/Source/Developer/NewLogVisualizer/Private/LogVisualizerSessionSettings.cpp
new file mode 100644
--- /dev/null
+++ b/Engine/Source/Developer/NewLogVisualizer/Private/LogVisualizerSessionSettings.cpp
@@ -0,0 +1,32 @@
+// Copyright 1998-2014 Epic Games, Inc. All Rights Reserved.
+
+#include "LogVisualizer.h"
+#include "Misc/CoreMisc.h"
+#include "LogVisualizerSessionSettings.h"
+#if WITH_EDITOR
+#include "Editor/EditorEngine.h"
+#include "ISettingsModule.h"
+#include "GeometryEdMode.h"
+#include "UnrealEdMisc.h"
+#endif // WITH_EDITOR
+
+ULogVisualizerSessionSettings::ULogVisualizerSessionSettings(const FObjectInitializer& ObjectInitializer) : Super(ObjectInitializer)
+{
+	bEnableGraphsVisualization = true;
+}
+
+#if WITH_EDITOR
+void ULogVisualizerSessionSettings::PostEditChangeProperty(struct FPropertyChangedEvent& PropertyChangedEvent)
+{
+	Super::PostEditChangeProperty(PropertyChangedEvent);
+
+	const FName Name = (PropertyChangedEvent.Property != nullptr) ? PropertyChangedEvent.Property->GetFName() : NAME_None;
+#if 0  //FIXME: should we save this settings too? (SebaK)
+	if (!FUnrealEdMisc::Get().IsDeletePreferences())
+	{
+		SaveConfig();
+	}
+#endif
+	SettingChangedEvent.Broadcast(Name);
+}
+#endif
diff --git a/Engine/Source/Developer/NewLogVisualizer/Private/LogVisualizerSettings.cpp b/Engine/Source/Developer/NewLogVisualizer/Private/LogVisualizerSettings.cpp
--- a/Engine/Source/Developer/NewLogVisualizer/Private/LogVisualizerSettings.cpp
+++ b/Engine/Source/Developer/NewLogVisualizer/Private/LogVisualizerSettings.cpp
@@ -15,6 +15,7 @@ ULogVisualizerSettings::ULogVisualizerSettings(const FObjectInitializer& ObjectI
 	TrivialLogsThreshold = 1;
 	DefaultCameraDistance = 150;
 	bSearchInsideLogs = true;
+	GraphsBackgroundColor = FColor(0, 0, 0, 70);
 }
 
 #if WITH_EDITOR
diff --git a/Engine/Source/Developer/NewLogVisualizer/Private/LogVisualizerStyle.cpp b/Engine/Source/Developer/NewLogVisualizer/Private/LogVisualizerStyle.cpp
--- a/Engine/Source/Developer/NewLogVisualizer/Private/LogVisualizerStyle.cpp
+++ b/Engine/Source/Developer/NewLogVisualizer/Private/LogVisualizerStyle.cpp
@@ -72,6 +72,7 @@ TSharedRef< FSlateStyleSet > FLogVisualizerStyle::Create()
 		Style.Set("Toolbar.Save", new IMAGE_BRUSH("Icons/LV_Save", Icon40x40));
 		Style.Set("Toolbar.Load", new IMAGE_BRUSH("Icons/LV_Load", Icon40x40));
 		Style.Set("Toolbar.Remove", new IMAGE_BRUSH("Icons/LV_Remove", Icon40x40));
+		Style.Set("Toolbar.Graphs", new IMAGE_BRUSH("Icons/icon_MatEd_Stats_40x", Icon40x40));
 	}
 
 	// Filters
diff --git a/Engine/Source/Developer/NewLogVisualizer/Private/STimeline.cpp b/Engine/Source/Developer/NewLogVisualizer/Private/STimeline.cpp
--- a/Engine/Source/Developer/NewLogVisualizer/Private/STimeline.cpp
+++ b/Engine/Source/Developer/NewLogVisualizer/Private/STimeline.cpp
@@ -231,3 +231,28 @@ const FSlateBrush* STimeline::GetBorder() const
 		return FLogVisualizerStyle::Get().GetBrush("ToolBar.Button.Normal");
 	}
 }
+
+void STimeline::AddEntry(const FVisualLogDevice::FVisualLogEntryItem& NewEntry) 
+{ 
+	if (Entries.Num())
+	{
+		FVisualLogDevice::FVisualLogEntryItem& LastEntry = Entries[Entries.Num() - 1];
+		if (LastEntry.Entry.TimeStamp == NewEntry.Entry.TimeStamp && LastEntry.Entry.Status.Num() == NewEntry.Entry.Status.Num())
+		{
+			LastEntry.Entry.Events.Append(NewEntry.Entry.Events);
+			LastEntry.Entry.LogLines.Append(NewEntry.Entry.LogLines);
+			LastEntry.Entry.ElementsToDraw.Append(NewEntry.Entry.ElementsToDraw);
+			LastEntry.Entry.HistogramSamples.Append(NewEntry.Entry.HistogramSamples);
+			LastEntry.Entry.DataBlocks.Append(NewEntry.Entry.DataBlocks);
+			//LastEntry.Entry.Status should be the same for the same timestamp;
+		}
+		else
+		{
+			Entries.Add(NewEntry);
+		}
+	}
+	else
+	{
+		Entries.Add(NewEntry);
+	}
+}
diff --git a/Engine/Source/Developer/NewLogVisualizer/Private/STimeline.h b/Engine/Source/Developer/NewLogVisualizer/Private/STimeline.h
--- a/Engine/Source/Developer/NewLogVisualizer/Private/STimeline.h
+++ b/Engine/Source/Developer/NewLogVisualizer/Private/STimeline.h
@@ -33,7 +33,7 @@ public:
 	void UpdateVisibility();
 	void UpdateVisibilityForItems();
 
-	void AddEntry(const FVisualLogDevice::FVisualLogEntryItem& Entry) { Entries.Add(Entry); }
+	void AddEntry(const FVisualLogDevice::FVisualLogEntryItem& Entry);
 	const TArray<FVisualLogDevice::FVisualLogEntryItem>& GetEntries() { return Entries; }
 	FName GetName() { return Name; }
 	void HandleLogVisualizerSettingChanged(FName Name);
diff --git a/Engine/Source/Developer/NewLogVisualizer/Private/SVisualLogger.cpp b/Engine/Source/Developer/NewLogVisualizer/Private/SVisualLogger.cpp
--- a/Engine/Source/Developer/NewLogVisualizer/Private/SVisualLogger.cpp
+++ b/Engine/Source/Developer/NewLogVisualizer/Private/SVisualLogger.cpp
@@ -188,6 +188,12 @@ void SVisualLogger::Construct(const FArguments& InArgs, const TSharedRef<SDockTa
 		FCanExecuteAction::CreateRaw(this, &SVisualLogger::HandleCameraCommandCanExecute), 
 		FIsActionChecked::CreateRaw(this, &SVisualLogger::HandleCameraCommandIsChecked),
 		FIsActionButtonVisible::CreateRaw(this, &SVisualLogger::HandleCameraCommandCanExecute));
+	ActionList.MapAction(Commands.ToggleGraphs,
+		FExecuteAction::CreateLambda([](){bool& bEnableGraphsVisualization = ULogVisualizerSessionSettings::StaticClass()->GetDefaultObject<ULogVisualizerSessionSettings>()->bEnableGraphsVisualization; bEnableGraphsVisualization = !bEnableGraphsVisualization; }),
+		FCanExecuteAction(),
+		FIsActionChecked::CreateLambda([]()->bool{return ULogVisualizerSessionSettings::StaticClass()->GetDefaultObject<ULogVisualizerSessionSettings>()->bEnableGraphsVisualization; }),
+		FIsActionButtonVisible());
+
 
 	// Tab Spawners
 	TabManager = FGlobalTabmanager::Get()->NewTabManager(ConstructUnderMajorTab);
diff --git a/Engine/Source/Developer/NewLogVisualizer/Private/SVisualLoggerToolbar.cpp b/Engine/Source/Developer/NewLogVisualizer/Private/SVisualLoggerToolbar.cpp
--- a/Engine/Source/Developer/NewLogVisualizer/Private/SVisualLoggerToolbar.cpp
+++ b/Engine/Source/Developer/NewLogVisualizer/Private/SVisualLoggerToolbar.cpp
@@ -31,7 +31,8 @@ TSharedRef<SWidget> SVisualLoggerToolbar::MakeToolbar(const TSharedRef<FUIComman
 		ToolBarBuilder.AddSeparator();
 		ToolBarBuilder.AddToolBarButton(FVisualLoggerCommands::Get().Load, NAME_None, LOCTEXT("Load", "Load"), TAttribute<FText>(), FSlateIcon(FLogVisualizerStyle::Get().GetStyleSetName(), TEXT("Toolbar.Load")));
 		ToolBarBuilder.AddToolBarButton(FVisualLoggerCommands::Get().Save, NAME_None, LOCTEXT("Save", "Save"), TAttribute<FText>(), FSlateIcon(FLogVisualizerStyle::Get().GetStyleSetName(), TEXT("Toolbar.Save")));
-		//ToolBarBuilder.AddToolBarButton(FMessagingDebuggerCommands::Get().ClearHistory, NAME_None, LOCTEXT("ClearHistory", "Clear History"));
+		ToolBarBuilder.AddSeparator();
+		ToolBarBuilder.AddToolBarButton(FVisualLoggerCommands::Get().ToggleGraphs, NAME_None, LOCTEXT("ToggleGraphs", "Toggle Graphs"), TAttribute<FText>(), FSlateIcon(FLogVisualizerStyle::Get().GetStyleSetName(), TEXT("Toolbar.Graphs")));
 	}
 
 	return ToolBarBuilder.MakeWidget();
diff --git a/Engine/Source/Developer/NewLogVisualizer/Private/VisualLoggerCanvasRenderer.cpp b/Engine/Source/Developer/NewLogVisualizer/Private/VisualLoggerCanvasRenderer.cpp
--- a/Engine/Source/Developer/NewLogVisualizer/Private/VisualLoggerCanvasRenderer.cpp
+++ b/Engine/Source/Developer/NewLogVisualizer/Private/VisualLoggerCanvasRenderer.cpp
@@ -101,7 +101,10 @@ void FVisualLoggerCanvasRenderer::DrawOnCanvas(class UCanvas* Canvas, class APla
 		}
 	}
 
-	DrawHistogramGraphs(Canvas, NULL);
+	if (ULogVisualizerSessionSettings::StaticClass()->GetDefaultObject<ULogVisualizerSessionSettings>()->bEnableGraphsVisualization)
+	{
+		DrawHistogramGraphs(Canvas, NULL);
+	}
 }
 
 void FVisualLoggerCanvasRenderer::DrawHistogramGraphs(class UCanvas* Canvas, class APlayerController*)
@@ -184,6 +187,7 @@ void FVisualLoggerCanvasRenderer::DrawHistogramGraphs(class UCanvas* Canvas, cla
 	int32 GraphIndex = 0;
 	if (CollectedGraphs.Num() > 0)
 	{
+		const FColor GraphsBackgroundColor = ULogVisualizerSettings::StaticClass()->GetDefaultObject<ULogVisualizerSettings>()->GraphsBackgroundColor;
 		const int NumberOfGraphs = CollectedGraphs.Num();
 		const int32 NumberOfColumns = FMath::CeilToInt(FMath::Sqrt(NumberOfGraphs));
 		int32 NumberOfRows = FMath::FloorToInt(NumberOfGraphs / NumberOfColumns);
@@ -260,7 +264,7 @@ void FVisualLoggerCanvasRenderer::DrawHistogramGraphs(class UCanvas* Canvas, cla
 			HistogramGraph->SetCursorLocation(SelectedEntry.TimeStamp);
 			HistogramGraph->SetNumThresholds(0);
 			HistogramGraph->SetStyles(EGraphAxisStyle::Grid, EGraphDataStyle::Lines);
-			HistogramGraph->SetBackgroundColor(FColor(0, 0, 0, 200));
+			HistogramGraph->SetBackgroundColor(GraphsBackgroundColor);
 			HistogramGraph->SetLegendPosition(/*bShowHistogramLabelsOutside*/ false ? ELegendPosition::Outside : ELegendPosition::Inside);
 			HistogramGraph->OffsetDataSets(/*bOffsetDataSet*/false);
 			HistogramGraph->bVisible = true;
diff --git a/Engine/Source/Developer/NewLogVisualizer/Private/VisualLoggerCommands.h b/Engine/Source/Developer/NewLogVisualizer/Private/VisualLoggerCommands.h
--- a/Engine/Source/Developer/NewLogVisualizer/Private/VisualLoggerCommands.h
+++ b/Engine/Source/Developer/NewLogVisualizer/Private/VisualLoggerCommands.h
@@ -27,6 +27,7 @@ public:
 		UI_COMMAND(FreeCamera, "FreeCamera", "Enable free camera", EUserInterfaceActionType::ToggleButton, FInputGesture());
 		UI_COMMAND(Load, "Load", "Load external vlogs", EUserInterfaceActionType::Button, FInputGesture());
 		UI_COMMAND(Save, "Save", "Save selected data to vlog file", EUserInterfaceActionType::Button, FInputGesture());
+		UI_COMMAND(ToggleGraphs, "ToggleGraphs", "Toggle graphs visualization on/off", EUserInterfaceActionType::ToggleButton, FInputGesture());
 	}
 
 public:
@@ -38,6 +39,7 @@ public:
 	TSharedPtr<FUICommandInfo> FreeCamera;
 	TSharedPtr<FUICommandInfo> Load;
 	TSharedPtr<FUICommandInfo> Save;
+	TSharedPtr<FUICommandInfo> ToggleGraphs;
 };
 
 
diff --git a/Engine/Source/Developer/NewLogVisualizer/Public/LogVisualizerSessionSettings.h b/Engine/Source/Developer/NewLogVisualizer/Public/LogVisualizerSessionSettings.h
new file mode 100644
--- /dev/null
+++ b/Engine/Source/Developer/NewLogVisualizer/Public/LogVisualizerSessionSettings.h
@@ -0,0 +1,32 @@
+// Copyright 1998-2015 Epic Games, Inc. All Rights Reserved.
+
+/*=============================================================================
+GameplayDebuggerSettings.h: Declares the UGameplayDebuggerSettings class.
+=============================================================================*/
+#pragma once
+
+
+#include "LogVisualizerSessionSettings.generated.h"
+
+UCLASS()
+class NEWLOGVISUALIZER_API ULogVisualizerSessionSettings : public UObject
+{
+	GENERATED_UCLASS_BODY()
+public:
+	/**Whether to show trivial logs, i.e. the ones with only one entry.*/
+	UPROPERTY(EditAnywhere, Category = "VisualLogger")
+	bool bEnableGraphsVisualization;
+
+	DECLARE_EVENT_OneParam(ULogVisualizerSessionSettings, FSettingChangedEvent, FName /*PropertyName*/);
+	FSettingChangedEvent& OnSettingChanged() { return SettingChangedEvent; }
+
+	// UObject overrides
+#if WITH_EDITOR
+	virtual void PostEditChangeProperty(struct FPropertyChangedEvent& PropertyChangedEvent) override;
+#endif
+
+private:
+
+	// Holds an event delegate that is executed when a setting has changed.
+	FSettingChangedEvent SettingChangedEvent;
+};
diff --git a/Engine/Source/Developer/NewLogVisualizer/Public/LogVisualizerSettings.h b/Engine/Source/Developer/NewLogVisualizer/Public/LogVisualizerSettings.h
--- a/Engine/Source/Developer/NewLogVisualizer/Public/LogVisualizerSettings.h
+++ b/Engine/Source/Developer/NewLogVisualizer/Public/LogVisualizerSettings.h
@@ -40,6 +40,10 @@ public:
 	UPROPERTY(EditAnywhere, config, Category = "VisualLogger")
 	bool bSearchInsideLogs;
 
+	/** Background color for 2d graphs visualization */
+	UPROPERTY(EditAnywhere, config, Category = "VisualLogger")
+	FColor GraphsBackgroundColor;
+
 	// UObject overrides
 #if WITH_EDITOR
 	virtual void PostEditChangeProperty(struct FPropertyChangedEvent& PropertyChangedEvent) override;
