Removed old LogVisualizer module (and old tool that way)
Fixed issue with visible 2d graphs after closing LogVisualizer
Removed LogVisualizer tool from experimentals so it's a regular tool inside Develoment Tools sub-menu now
Few other minor issue fixes in LogVisualizer tool

diff --git a/Engine/Source/Developer/LogVisualizer/LogVisualizer.Build.cs b/Engine/Source/Developer/LogVisualizer/LogVisualizer.Build.cs
deleted file mode 100644
diff --git a/Engine/Source/Developer/LogVisualizer/Private/LogVisualizer.cpp b/Engine/Source/Developer/LogVisualizer/Private/LogVisualizer.cpp
deleted file mode 100644
diff --git a/Engine/Source/Developer/LogVisualizer/Private/LogVisualizer.h b/Engine/Source/Developer/LogVisualizer/Private/LogVisualizer.h
deleted file mode 100644
diff --git a/Engine/Source/Developer/LogVisualizer/Private/LogVisualizerDebugActor.cpp b/Engine/Source/Developer/LogVisualizer/Private/LogVisualizerDebugActor.cpp
deleted file mode 100644
diff --git a/Engine/Source/Developer/LogVisualizer/Private/LogVisualizerDebugActor.h b/Engine/Source/Developer/LogVisualizer/Private/LogVisualizerDebugActor.h
deleted file mode 100644
diff --git a/Engine/Source/Developer/LogVisualizer/Private/LogVisualizerModule.cpp b/Engine/Source/Developer/LogVisualizer/Private/LogVisualizerModule.cpp
deleted file mode 100644
diff --git a/Engine/Source/Developer/LogVisualizer/Private/LogVisualizerPCH.h b/Engine/Source/Developer/LogVisualizer/Private/LogVisualizerPCH.h
deleted file mode 100644
diff --git a/Engine/Source/Developer/LogVisualizer/Private/SFilterList.cpp b/Engine/Source/Developer/LogVisualizer/Private/SFilterList.cpp
deleted file mode 100644
diff --git a/Engine/Source/Developer/LogVisualizer/Private/SFilterList.h b/Engine/Source/Developer/LogVisualizer/Private/SFilterList.h
deleted file mode 100644
diff --git a/Engine/Source/Developer/LogVisualizer/Private/SLogBar.cpp b/Engine/Source/Developer/LogVisualizer/Private/SLogBar.cpp
deleted file mode 100644
diff --git a/Engine/Source/Developer/LogVisualizer/Private/SLogBar.h b/Engine/Source/Developer/LogVisualizer/Private/SLogBar.h
deleted file mode 100644
diff --git a/Engine/Source/Developer/LogVisualizer/Private/SLogVisualizer.cpp b/Engine/Source/Developer/LogVisualizer/Private/SLogVisualizer.cpp
deleted file mode 100644
diff --git a/Engine/Source/Developer/LogVisualizer/Private/SLogVisualizer.h b/Engine/Source/Developer/LogVisualizer/Private/SLogVisualizer.h
deleted file mode 100644
diff --git a/Engine/Source/Developer/LogVisualizer/Private/SLogVisualizerReport.cpp b/Engine/Source/Developer/LogVisualizer/Private/SLogVisualizerReport.cpp
deleted file mode 100644
diff --git a/Engine/Source/Developer/LogVisualizer/Private/SLogVisualizerReport.h b/Engine/Source/Developer/LogVisualizer/Private/SLogVisualizerReport.h
deleted file mode 100644
diff --git a/Engine/Source/Developer/LogVisualizer/Private/SLogsTableRow.cpp b/Engine/Source/Developer/LogVisualizer/Private/SLogsTableRow.cpp
deleted file mode 100644
diff --git a/Engine/Source/Developer/LogVisualizer/Private/SLogsTableRow.h b/Engine/Source/Developer/LogVisualizer/Private/SLogsTableRow.h
deleted file mode 100644
diff --git a/Engine/Source/Developer/LogVisualizer/Public/ILogVisualizer.h b/Engine/Source/Developer/LogVisualizer/Public/ILogVisualizer.h
deleted file mode 100644
diff --git a/Engine/Source/Developer/LogVisualizer/Public/LogVisualizerModule.h b/Engine/Source/Developer/LogVisualizer/Public/LogVisualizerModule.h
deleted file mode 100644
diff --git a/Engine/Source/Developer/NewLogVisualizer/NewLogVisualizer.Build.cs b/Engine/Source/Developer/NewLogVisualizer/NewLogVisualizer.Build.cs
--- a/Engine/Source/Developer/NewLogVisualizer/NewLogVisualizer.Build.cs
+++ b/Engine/Source/Developer/NewLogVisualizer/NewLogVisualizer.Build.cs
@@ -18,7 +18,8 @@ public class NewLogVisualizer : ModuleRules
 			new string[] {
 				"Developer/TaskGraph/Public",
 				"Runtime/Engine/Classes",
-                "Editor/GeometryMode/Public"
+                "Editor/GeometryMode/Public",
+                "Editor/WorkspaceMenuStructure/Public"
 			}
 		);
 
diff --git a/Engine/Source/Developer/NewLogVisualizer/Private/LogVisualizerModule.cpp b/Engine/Source/Developer/NewLogVisualizer/Private/LogVisualizerModule.cpp
--- a/Engine/Source/Developer/NewLogVisualizer/Private/LogVisualizerModule.cpp
+++ b/Engine/Source/Developer/NewLogVisualizer/Private/LogVisualizerModule.cpp
@@ -10,6 +10,7 @@
 #include "Editor/EditorEngine.h"
 #include "ISettingsModule.h"
 #include "LevelEditor.h"
+#include "WorkspaceMenuStructureModule.h"
 #endif // WITH_EDITOR
 
 #define LOCTEXT_NAMESPACE "FLogVisualizerModule"
@@ -25,12 +26,14 @@ void FNewLogVisualizerModule::StartupModule()
 	FVisualLoggerCommands::Register();
 	IModularFeatures::Get().RegisterModularFeature(VisualLoggerTabName, this);
 
-	// This is still experimental in the editor, so it'll be invoked specifically in FMainMenu if the experimental settings flag is set.
-	// When no longer experimental, switch to the nomad spawner registration below
-	FGlobalTabmanager::Get()->RegisterTabSpawner(VisualLoggerTabName, FOnSpawnTab::CreateRaw(this, &FNewLogVisualizerModule::SpawnLogVisualizerTab));
+	FGlobalTabmanager::Get()->RegisterNomadTabSpawner(
+		VisualLoggerTabName, 
+		FOnSpawnTab::CreateRaw(this, &FNewLogVisualizerModule::SpawnLogVisualizerTab))
+		.SetGroup(WorkspaceMenu::GetMenuStructure().GetDeveloperToolsMiscCategory())
+		.SetDisplayName(NSLOCTEXT("LogVisualizerApp", "TabTitle", "Visual Logger"))
+		.SetTooltipText(NSLOCTEXT("LogVisualizerApp", "TooltipText", "Opens Visual Logger tool."))
+		.SetIcon(FSlateIcon(FLogVisualizerStyle::GetStyleSetName(), "LogVisualizerApp.TabIcon"));
 
-	UEditorExperimentalSettings::StaticClass()->GetDefaultObject<UEditorExperimentalSettings>()->OnSettingChanged().AddRaw(this, &FNewLogVisualizerModule::HandleExperimentalSettingChanged);
-	
 	ISettingsModule* SettingsModule = FModuleManager::GetModulePtr<ISettingsModule>("Settings");
 	if (SettingsModule != nullptr)
 	{
@@ -48,8 +51,6 @@ void FNewLogVisualizerModule::ShutdownModule()
 	FVisualLoggerCommands::Unregister();
 	IModularFeatures::Get().UnregisterModularFeature(VisualLoggerTabName, this);
 	
-	//UEditorExperimentalSettings::StaticClass()->GetDefaultObject<UEditorExperimentalSettings>()->OnSettingChanged().RemoveAll(this);
-
 	ISettingsModule* SettingsModule = FModuleManager::GetModulePtr<ISettingsModule>("Settings");
 	if (SettingsModule != nullptr)
 	{
@@ -59,36 +60,6 @@ void FNewLogVisualizerModule::ShutdownModule()
 	FLogVisualizerStyle::Shutdown();
 }
 
-#if WITH_EDITOR
-
-void FNewLogVisualizerModule::HandleExperimentalSettingChanged(FName PropertyName)
-{
-	if (PropertyName == TEXT("bVisualLogger"))
-	{
-		if (UEditorExperimentalSettings::StaticClass()->GetDefaultObject<UEditorExperimentalSettings>()->bVisualLogger == false)
-		{
-			ISettingsModule* SettingsModule = FModuleManager::GetModulePtr<ISettingsModule>("Settings");
-			if (SettingsModule != nullptr)
-			{
-				SettingsModule->UnregisterSettings("Editor", "General", "VisualLogger");
-			}
-		}
-		else
-		{
-			ISettingsModule* SettingsModule = FModuleManager::GetModulePtr<ISettingsModule>("Settings");
-			if (SettingsModule != nullptr)
-			{
-				SettingsModule->RegisterSettings("Editor", "General", "VisualLogger",
-					LOCTEXT("AIToolsSettingsName", "Visual Logger"),
-					LOCTEXT("AIToolsSettingsDescription", "General settings for UE4 AI Tools."),
-					ULogVisualizerSettings::StaticClass()->GetDefaultObject()
-					);
-			}
-		}
-	}
-}
-#endif
-
 TSharedRef<SDockTab> FNewLogVisualizerModule::SpawnLogVisualizerTab(const FSpawnTabArgs& SpawnTabArgs)
 {
 	const TSharedRef<SDockTab> MajorTab = SNew(SDockTab)
@@ -103,7 +74,7 @@ TSharedRef<SDockTab> FNewLogVisualizerModule::SpawnLogVisualizerTab(const FSpawn
 	return MajorTab;
 }
 
-void FNewLogVisualizerModule::OnTabClosed(TSharedRef<SDockTab>)
+void FNewLogVisualizerModule::OnTabClosed(TSharedRef<SDockTab> DockTab)
 {
 	UWorld* World = NULL;
 #if WITH_EDITOR
@@ -122,6 +93,9 @@ void FNewLogVisualizerModule::OnTabClosed(TSharedRef<SDockTab>)
 		World = GEngine->GetWorld();
 	}
 
+	TSharedRef<SVisualLogger> VisualLoggerTab = StaticCastSharedRef<SVisualLogger>(DockTab->GetContent());
+	VisualLoggerTab->OnTabLosed();
+
 	if (World == NULL)
 	{
 		World = GWorld;
diff --git a/Engine/Source/Developer/NewLogVisualizer/Private/LogVisualizerStyle.cpp b/Engine/Source/Developer/NewLogVisualizer/Private/LogVisualizerStyle.cpp
--- a/Engine/Source/Developer/NewLogVisualizer/Private/LogVisualizerStyle.cpp
+++ b/Engine/Source/Developer/NewLogVisualizer/Private/LogVisualizerStyle.cpp
@@ -48,6 +48,8 @@ TSharedRef< FSlateStyleSet > FLogVisualizerStyle::Create()
 	FSlateStyleSet& Style = StyleRef.Get();
 	// Generic styles
 	{
+		Style.Set("LogVisualizerApp.TabIcon", new IMAGE_BRUSH("Icons/icon_tab_DebugTools_40x", Icon16x16));
+		
 		Style.Set("LogVisualizer.LogBar.Background", new BOX_BRUSH("Common/ProgressBar_Background", FMargin(5.f / 12.f)));
 		Style.Set("LogVisualizer.LogBar.Selected", new BOX_BRUSH("Common/TaskGraph_Selected", FMargin(5.f / 12.f)));
 		Style.Set("LogVisualizer.LogBar.EntryDefault", new BOX_BRUSH("Common/TaskGraph_Mono", FMargin(5.f / 12.f)));
diff --git a/Engine/Source/Developer/NewLogVisualizer/Private/STimeline.cpp b/Engine/Source/Developer/NewLogVisualizer/Private/STimeline.cpp
--- a/Engine/Source/Developer/NewLogVisualizer/Private/STimeline.cpp
+++ b/Engine/Source/Developer/NewLogVisualizer/Private/STimeline.cpp
@@ -231,28 +231,3 @@ const FSlateBrush* STimeline::GetBorder() const
 		return FLogVisualizerStyle::Get().GetBrush("ToolBar.Button.Normal");
 	}
 }
-
-void STimeline::AddEntry(const FVisualLogDevice::FVisualLogEntryItem& NewEntry) 
-{ 
-	if (Entries.Num())
-	{
-		FVisualLogDevice::FVisualLogEntryItem& LastEntry = Entries[Entries.Num() - 1];
-		if (LastEntry.Entry.TimeStamp == NewEntry.Entry.TimeStamp && LastEntry.Entry.Status.Num() == NewEntry.Entry.Status.Num())
-		{
-			LastEntry.Entry.Events.Append(NewEntry.Entry.Events);
-			LastEntry.Entry.LogLines.Append(NewEntry.Entry.LogLines);
-			LastEntry.Entry.ElementsToDraw.Append(NewEntry.Entry.ElementsToDraw);
-			LastEntry.Entry.HistogramSamples.Append(NewEntry.Entry.HistogramSamples);
-			LastEntry.Entry.DataBlocks.Append(NewEntry.Entry.DataBlocks);
-			//LastEntry.Entry.Status should be the same for the same timestamp;
-		}
-		else
-		{
-			Entries.Add(NewEntry);
-		}
-	}
-	else
-	{
-		Entries.Add(NewEntry);
-	}
-}
diff --git a/Engine/Source/Developer/NewLogVisualizer/Private/STimeline.h b/Engine/Source/Developer/NewLogVisualizer/Private/STimeline.h
--- a/Engine/Source/Developer/NewLogVisualizer/Private/STimeline.h
+++ b/Engine/Source/Developer/NewLogVisualizer/Private/STimeline.h
@@ -33,7 +33,7 @@ public:
 	void UpdateVisibility();
 	void UpdateVisibilityForItems();
 
-	void AddEntry(const FVisualLogDevice::FVisualLogEntryItem& Entry);
+	void AddEntry(const FVisualLogDevice::FVisualLogEntryItem& Entry) { Entries.Add(Entry); }
 	const TArray<FVisualLogDevice::FVisualLogEntryItem>& GetEntries() { return Entries; }
 	FName GetName() { return Name; }
 	void HandleLogVisualizerSettingChanged(FName Name);
diff --git a/Engine/Source/Developer/NewLogVisualizer/Private/SVisualLogger.cpp b/Engine/Source/Developer/NewLogVisualizer/Private/SVisualLogger.cpp
--- a/Engine/Source/Developer/NewLogVisualizer/Private/SVisualLogger.cpp
+++ b/Engine/Source/Developer/NewLogVisualizer/Private/SVisualLogger.cpp
@@ -148,6 +148,11 @@ SVisualLogger::SVisualLogger() : SCompoundWidget(), CommandList(MakeShareable(ne
 }
 
 SVisualLogger::~SVisualLogger()
+{
+
+}
+
+void SVisualLogger::OnTabLosed()
 {
 	UDebugDrawService::Unregister(FDebugDrawDelegate::CreateRaw(VisualLoggerCanvasRenderer.Get(), &FVisualLoggerCanvasRenderer::DrawOnCanvas));
 	VisualLoggerCanvasRenderer = NULL;
diff --git a/Engine/Source/Developer/NewLogVisualizer/Private/SVisualLogger.h b/Engine/Source/Developer/NewLogVisualizer/Private/SVisualLogger.h
--- a/Engine/Source/Developer/NewLogVisualizer/Private/SVisualLogger.h
+++ b/Engine/Source/Developer/NewLogVisualizer/Private/SVisualLogger.h
@@ -52,6 +52,7 @@ public:
 
 	/** Callback for for when the owner tab's visual state is being persisted. */
 	void HandleMajorTabPersistVisualState();
+	void OnTabLosed();
 
 	void OnNewLogEntry(const FVisualLogDevice::FVisualLogEntryItem& Entry);
 	void CollectNewCategories(const FVisualLogDevice::FVisualLogEntryItem& Entry);
diff --git a/Engine/Source/Developer/NewLogVisualizer/Private/SVisualLoggerToolbar.cpp b/Engine/Source/Developer/NewLogVisualizer/Private/SVisualLoggerToolbar.cpp
--- a/Engine/Source/Developer/NewLogVisualizer/Private/SVisualLoggerToolbar.cpp
+++ b/Engine/Source/Developer/NewLogVisualizer/Private/SVisualLoggerToolbar.cpp
@@ -32,7 +32,7 @@ TSharedRef<SWidget> SVisualLoggerToolbar::MakeToolbar(const TSharedRef<FUIComman
 		ToolBarBuilder.AddToolBarButton(FVisualLoggerCommands::Get().Load, NAME_None, LOCTEXT("Load", "Load"), TAttribute<FText>(), FSlateIcon(FLogVisualizerStyle::Get().GetStyleSetName(), TEXT("Toolbar.Load")));
 		ToolBarBuilder.AddToolBarButton(FVisualLoggerCommands::Get().Save, NAME_None, LOCTEXT("Save", "Save"), TAttribute<FText>(), FSlateIcon(FLogVisualizerStyle::Get().GetStyleSetName(), TEXT("Toolbar.Save")));
 		ToolBarBuilder.AddSeparator();
-		ToolBarBuilder.AddToolBarButton(FVisualLoggerCommands::Get().ToggleGraphs, NAME_None, LOCTEXT("ToggleGraphs", "Toggle Graphs"), TAttribute<FText>(), FSlateIcon(FLogVisualizerStyle::Get().GetStyleSetName(), TEXT("Toolbar.Graphs")));
+		ToolBarBuilder.AddToolBarButton(FVisualLoggerCommands::Get().ToggleGraphs, NAME_None, LOCTEXT("ToggleGraphs", "Graphs"), TAttribute<FText>(), FSlateIcon(FLogVisualizerStyle::Get().GetStyleSetName(), TEXT("Toolbar.Graphs")));
 	}
 
 	return ToolBarBuilder.MakeWidget();
diff --git a/Engine/Source/Developer/NewLogVisualizer/Public/LogVisualizerModule.h b/Engine/Source/Developer/NewLogVisualizer/Public/LogVisualizerModule.h
--- a/Engine/Source/Developer/NewLogVisualizer/Public/LogVisualizerModule.h
+++ b/Engine/Source/Developer/NewLogVisualizer/Public/LogVisualizerModule.h
@@ -17,7 +17,6 @@ public:
 	// End IModuleInterface
 
 private:
-	void HandleExperimentalSettingChanged(FName PropertyName);
 	TSharedRef<SDockTab> SpawnLogVisualizerTab(const FSpawnTabArgs& SpawnTabArgs);
 	void OnTabClosed(TSharedRef<SDockTab>);
 };
diff --git a/Engine/Source/Editor/EditorStyle/Private/SlateEditorStyle.cpp b/Engine/Source/Editor/EditorStyle/Private/SlateEditorStyle.cpp
--- a/Engine/Source/Editor/EditorStyle/Private/SlateEditorStyle.cpp
+++ b/Engine/Source/Editor/EditorStyle/Private/SlateEditorStyle.cpp
@@ -1548,21 +1548,6 @@ void FSlateEditorStyle::FStyle::SetupGeneralStyles()
 		Set( "CollisionAnalyzer.Group", new IMAGE_BRUSH( "Icons/CA_Group", FVector2D(10,18) ) );
 		Set( "CollisionAnalyzer.GroupBackground", new BOX_BRUSH( "/Icons/CA_GroupBackground", FMargin(4.f/16.f) ) );
 	}
-
-	// Log Visualizer
-	{
-		Set("LogVisualizer.Pause", new IMAGE_BRUSH("Icons/icon_pause_40x", Icon48x48));
-		Set( "LogVisualizer.LogBar.Background", new BOX_BRUSH( "Common/ProgressBar_Background", FMargin(5.f/12.f) ) );
-		Set( "LogVisualizer.LogBar.Selected", new BOX_BRUSH( "Common/TaskGraph_Selected", FMargin(5.f/12.f) ) );
-		Set( "LogVisualizer.LogBar.EntryDefault", new BOX_BRUSH( "Common/TaskGraph_Mono", FMargin(5.f/12.f) ) );
-		Set( "LogVisualizer.LogBar.TimeMark", new BOX_BRUSH( "Icons/LV_BarMark", FMargin(5.f/12.f) ) );
-		Set("LogVisualizer.Record", new IMAGE_BRUSH("Icons/LV_Record", Icon40x40));
-		Set("LogVisualizer.Stop", new IMAGE_BRUSH("Icons/LV_Stop", Icon40x40));
-		Set("LogVisualizer.Camera", new IMAGE_BRUSH("Icons/LV_Camera", Icon48x48));
-		Set( "LogVisualizer.Save", new IMAGE_BRUSH( "Icons/LV_Save", Icon40x40 ) );
-		Set( "LogVisualizer.Load", new IMAGE_BRUSH( "Icons/LV_Load", Icon40x40 ) );
-		Set( "LogVisualizer.Remove", new IMAGE_BRUSH( "Icons/LV_Remove", Icon40x40 ) );
-	}
 #endif // WITH_EDITOR || IS_PROGRAM
 
 #if WITH_EDITOR
diff --git a/Engine/Source/Editor/MainFrame/Private/Menus/MainMenu.cpp b/Engine/Source/Editor/MainFrame/Private/Menus/MainMenu.cpp
--- a/Engine/Source/Editor/MainFrame/Private/Menus/MainMenu.cpp
+++ b/Engine/Source/Editor/MainFrame/Private/Menus/MainMenu.cpp
@@ -173,10 +173,9 @@ void FMainMenu::FillWindowMenu( FMenuBuilder& MenuBuilder, const TSharedRef< FEx
 		bool bMessagingDebugger = GetDefault<UEditorExperimentalSettings>()->bMessagingDebugger;
 		bool bBlutility = GetDefault<UEditorExperimentalSettings>()->bEnableEditorUtilityBlueprints;
 		bool bTranslationEditor = GetDefault<UEditorExperimentalSettings>()->bEnableTranslationEditor;
-		bool bVisualLogger = GetDefault<UEditorExperimentalSettings>()->bVisualLogger;
 
 		// Make sure at least one is enabled before creating the section
-		if (bProjectLauncher || bMessagingDebugger || bBlutility || bTranslationEditor || bVisualLogger)
+		if (bProjectLauncher || bMessagingDebugger || bBlutility || bTranslationEditor)
 		{
 			MenuBuilder.BeginSection("ExperimentalTabSpawners", LOCTEXT("ExperimentalTabSpawnersHeading", "Experimental"));
 			{
@@ -202,17 +201,6 @@ void FMainMenu::FillWindowMenu( FMenuBuilder& MenuBuilder, const TSharedRef< FEx
 						);
 				}
 
-				// Visual Logger visualizer
-				if (bVisualLogger)
-				{
-					MenuBuilder.AddMenuEntry(
-						LOCTEXT("VisualLoggerLabel", "Visual Logger"),
-						LOCTEXT("VisualLoggerToolTip", "Enables Visual Logger visualizer and viewer."),
-						FSlateIcon(), // Icon lives in the plugin dir for the debugger
-						FUIAction(FExecuteAction::CreateStatic(&FMainMenu::OpenVisualLogger))
-						);
-				}
-
 				// Blutility
 				if (bBlutility)
 				{
diff --git a/Engine/Source/Editor/UnrealEd/Classes/Settings/EditorExperimentalSettings.h b/Engine/Source/Editor/UnrealEd/Classes/Settings/EditorExperimentalSettings.h
--- a/Engine/Source/Editor/UnrealEd/Classes/Settings/EditorExperimentalSettings.h
+++ b/Engine/Source/Editor/UnrealEd/Classes/Settings/EditorExperimentalSettings.h
@@ -85,10 +85,6 @@ public:
 	UPROPERTY(EditAnywhere, config, Category = Blueprints, meta = (DisplayName = "Enable Blueprint Merge Tool"))
 	bool bEnableBlueprintMergeTool;
 
-	/** Enables Visual Logger visualizer/viewer */
-	UPROPERTY(EditAnywhere, config, Category = Tools, meta = (DisplayName = "Visual Logger"))
-	bool bVisualLogger;
-
 	/** Enables Gameplay Debugger */
 	UPROPERTY(EditAnywhere, config, Category = Tools, meta = (DisplayName = "Gameplay Debugger"))
 	bool bGameplayDebugger;
diff --git a/Engine/Source/Editor/UnrealEd/Private/Editor.cpp b/Engine/Source/Editor/UnrealEd/Private/Editor.cpp
--- a/Engine/Source/Editor/UnrealEd/Private/Editor.cpp
+++ b/Engine/Source/Editor/UnrealEd/Private/Editor.cpp
@@ -606,11 +606,6 @@ void UEditorEngine::Init(IEngineLoop* InEngineLoop)
 			FModuleManager::Get().LoadModule(TEXT("EnvironmentQueryEditor"));
 		}
 
-		if (GetDefault<UEditorExperimentalSettings>()->bVisualLogger)
-		{
-			FModuleManager::Get().LoadModule(TEXT("NewLogVisualizer"));
-		}
-
 		bool bGameplayAbilitiesEnabled = false;
 		GConfig->GetBool(TEXT("GameplayAbilities"), TEXT("GameplayAbilitiesEditorEnabled"), bGameplayAbilitiesEnabled, GEngineIni);
 		if (bGameplayAbilitiesEnabled)
@@ -618,6 +613,8 @@ void UEditorEngine::Init(IEngineLoop* InEngineLoop)
 			FModuleManager::Get().LoadModule(TEXT("GameplayAbilitiesEditor"));
 		}
 
+
+		FModuleManager::Get().LoadModule(TEXT("NewLogVisualizer"));
 		FModuleManager::Get().LoadModule(TEXT("HotReload"));
 	}
 
diff --git a/Engine/Source/Editor/UnrealEd/Private/Settings/SettingsClasses.cpp b/Engine/Source/Editor/UnrealEd/Private/Settings/SettingsClasses.cpp
--- a/Engine/Source/Editor/UnrealEd/Private/Settings/SettingsClasses.cpp
+++ b/Engine/Source/Editor/UnrealEd/Private/Settings/SettingsClasses.cpp
@@ -59,7 +59,6 @@ UEditorExperimentalSettings::UEditorExperimentalSettings( const FObjectInitializ
 void UEditorExperimentalSettings::PostEditChangeProperty( struct FPropertyChangedEvent& PropertyChangedEvent )
 {
 	static const FName NAME_EQS = GET_MEMBER_NAME_CHECKED(UEditorExperimentalSettings, bEQSEditor);
-	static const FName NAME_VISUAL_LOGGER = GET_MEMBER_NAME_CHECKED(UEditorExperimentalSettings, bVisualLogger);
 
 	Super::PostEditChangeProperty(PropertyChangedEvent);
 
@@ -76,13 +75,6 @@ void UEditorExperimentalSettings::PostEditChangeProperty( struct FPropertyChange
 			FModuleManager::Get().LoadModule(TEXT("EnvironmentQueryEditor"));
 		}
 	}
-	else if (Name == NAME_VISUAL_LOGGER)
-	{
-		if (bVisualLogger)
-		{
-			FModuleManager::Get().LoadModule(TEXT("NewLogVisualizer"));
-		}
-	}
 
 	if (!FUnrealEdMisc::Get().IsDeletePreferences())
 	{
diff --git a/Engine/Source/Runtime/Engine/Engine.Build.cs b/Engine/Source/Runtime/Engine/Engine.Build.cs
--- a/Engine/Source/Runtime/Engine/Engine.Build.cs
+++ b/Engine/Source/Runtime/Engine/Engine.Build.cs
@@ -123,9 +123,6 @@ public class Engine : ModuleRules
 				PrivateDependencyModuleNames.Add("CollisionAnalyzer");
 				CircularlyReferencedDependentModules.Add("CollisionAnalyzer");
 
-                PrivateDependencyModuleNames.Add("LogVisualizer");
-                CircularlyReferencedDependentModules.Add("LogVisualizer");
-            
                 PrivateDependencyModuleNames.Add("NewLogVisualizer");
                 CircularlyReferencedDependentModules.Add("NewLogVisualizer");
             }
diff --git a/Engine/Source/Runtime/Engine/Private/CheatManager.cpp b/Engine/Source/Runtime/Engine/Private/CheatManager.cpp
--- a/Engine/Source/Runtime/Engine/Private/CheatManager.cpp
+++ b/Engine/Source/Runtime/Engine/Private/CheatManager.cpp
@@ -1118,23 +1118,4 @@ void UCheatManager::LogOutBugItGoToLogFile( const FString& InScreenShotDesc, con
 #endif // ALLOW_DEBUG_FILES
 }
 
-#if WITH_EDITOR
-#include "Developer/LogVisualizer/Public/LogVisualizerModule.h"
-static class FOldLogVisualizerExec : private FSelfRegisteringExec
-{
-public:
-	/** Console commands, see embeded usage statement **/
-	virtual bool Exec(UWorld* InWorld, const TCHAR* Cmd, FOutputDevice& Ar) override
-	{
-		if (FParse::Command(&Cmd, TEXT("OLDVISLOG")))
-		{
-			FLogVisualizerModule::Get()->SummonUI(InWorld);
-			return true;
-		}
-		return false;
-	}
-} OldLogVisualizerExec;
-#endif
-
-
 #undef LOCTEXT_NAMESPACE
